ShotManagement: Fix invalid asset data used for the No Schema option when creating Cine Assemblies.

#rnx

diff --git a/Engine/Plugins/VirtualProduction/CinematicAssemblyTools/Source/ShotsEditor/Private/UI/CineAssembly/SCineAssemblyConfigWindow.cpp b/Engine/Plugins/VirtualProduction/CinematicAssemblyTools/Source/ShotsEditor/Private/UI/CineAssembly/SCineAssemblyConfigWindow.cpp
--- a/Engine/Plugins/VirtualProduction/CinematicAssemblyTools/Source/ShotsEditor/Private/UI/CineAssembly/SCineAssemblyConfigWindow.cpp
+++ b/Engine/Plugins/VirtualProduction/CinematicAssemblyTools/Source/ShotsEditor/Private/UI/CineAssembly/SCineAssemblyConfigWindow.cpp
@@ -6,6 +6,7 @@
 #include "ContentBrowserModule.h"
 #include "Editor.h"
 #include "IContentBrowserSingleton.h"
+#include "Interfaces/IPluginManager.h"
 #include "IStructureDetailsView.h"
 #include "Modules/ModuleManager.h"
 #include "NamingTokensEngineSubsystem.h"
@@ -131,8 +132,11 @@ TSharedRef<SWidget> SCineAssemblyConfigWindow::MakeCineTemplatePanel()
 	Config.bShowPathInColumnView = true;
 	Config.OnAssetSelected = FOnAssetSelected::CreateSP(this, &SCineAssemblyConfigWindow::OnSchemaSelected);
 
-	const FName NoSchemaName = "No Schema";
-	FAssetData NoSchemaAssetData = FAssetData("/", "", NoSchemaName, FTopLevelAssetPath(""));
+	const FString PackageName = IPluginManager::Get().FindPlugin(UE_PLUGIN_NAME)->GetMountedAssetPath();
+	const FString NoSchemaName = TEXT("No Schema");
+	const FString PackagePath = PackageName / NoSchemaName;
+
+	FAssetData NoSchemaAssetData = FAssetData(*PackagePath, *PackageName, *NoSchemaName, FTopLevelAssetPath());
 
 	// Add a fake asset to the list (so that it appears as a tile in the asset picker) that represents a selection of no schema
 	Config.OnGetCustomSourceAssets = FOnGetCustomSourceAssets::CreateLambda([NoSchemaAssetData](const FARFilter& Filter, TArray<FAssetData>& OutAssets)
