ShotManagement: Fix schema hierarhcy view to allow placement of the top-level assembly in a different folder. This will result in a saved DefaultAssemblyPath associated with the schema, that will be used to change the package path of assemblies created from that schema.
#rnx

diff --git a/Engine/Plugins/VirtualProduction/CinematicAssemblyTools/Source/ShotsEditor/Private/CineAssembly.cpp b/Engine/Plugins/VirtualProduction/CinematicAssemblyTools/Source/ShotsEditor/Private/CineAssembly.cpp
--- a/Engine/Plugins/VirtualProduction/CinematicAssemblyTools/Source/ShotsEditor/Private/CineAssembly.cpp
+++ b/Engine/Plugins/VirtualProduction/CinematicAssemblyTools/Source/ShotsEditor/Private/CineAssembly.cpp
@@ -117,7 +117,11 @@ void UCineAssembly::CreateSubAssemblies()
 
 	// Get the path where the top-level assembly will be created so we can create other assets relative to it
 	const FAssetData AssemblyAssetData = FAssetData(this);
-	const FString PackagePath = AssemblyAssetData.PackagePath.ToString();
+	FString PackagePath = AssemblyAssetData.PackagePath.ToString();
+
+	// Remove the default assembly path from the top-level assembly's package path to get back to the "root" path for the schema's folder hierarchy
+	FNamingTokenResultData DefaultAssemblyPathResult = NamingTokensSubsystem->EvaluateTokenString(BaseSchema->DefaultAssemblyPath, FilterArgs, { NamingTokenContext });
+	PackagePath = PackagePath.Replace(*DefaultAssemblyPathResult.EvaluatedText.ToString(), TEXT(""));
 
 	// Create the default folders for this assembly, based on the schema
 	for (FTemplateString& FolderPath : DefaultFolderNames)
diff --git a/Engine/Plugins/VirtualProduction/CinematicAssemblyTools/Source/ShotsEditor/Private/CineAssemblyFactory.cpp b/Engine/Plugins/VirtualProduction/CinematicAssemblyTools/Source/ShotsEditor/Private/CineAssemblyFactory.cpp
--- a/Engine/Plugins/VirtualProduction/CinematicAssemblyTools/Source/ShotsEditor/Private/CineAssemblyFactory.cpp
+++ b/Engine/Plugins/VirtualProduction/CinematicAssemblyTools/Source/ShotsEditor/Private/CineAssemblyFactory.cpp
@@ -4,8 +4,10 @@
 
 #include "AssetToolsModule.h"
 #include "CineAssemblyNamingTokens.h"
+#include "ContentBrowserModule.h"
 #include "Framework/Application/SlateApplication.h"
 #include "Interfaces/IMainFrameModule.h"
+#include "IContentBrowserSingleton.h"
 #include "Modules/ModuleManager.h"
 #include "MovieScene.h"
 #include "MovieSceneToolsProjectSettings.h"
@@ -32,15 +34,28 @@ bool UCineAssemblyFactory::CanCreateNew() const
 
 UObject* UCineAssemblyFactory::FactoryCreateNew(UClass* Class, UObject* InParent, FName Name, EObjectFlags Flags, UObject* Context, FFeedbackContext* Warn)
 {
+	IAssetTools& AssetTools = FModuleManager::GetModuleChecked<FAssetToolsModule>("AssetTools").Get();
+
 	if (CineAssemblyToConfigure)
 	{
+		if (const UCineAssemblySchema* Schema = CineAssemblyToConfigure->GetSchema())
+		{
+			if (!Schema->DefaultAssemblyPath.IsEmpty())
+			{
+				const FString ResolvedString = ResolveTokenString(Schema->DefaultAssemblyPath);
+
+				// Change the outer package to include the DefaultAssemblyPath
+				const FString PackageName = FPaths::GetPath(InParent->GetPathName()) / ResolvedString / Name.ToString();
+				InParent = CreatePackage(*PackageName);
+			}
+		}
+
 		// Update the name, outer package, and flags of the configuration asset to no longer be transient
 		CineAssemblyToConfigure->Rename(*Name.ToString(), InParent);
 		CineAssemblyToConfigure->SetFlags(Flags);
 		CineAssemblyToConfigure->ClearFlags(RF_Transient);
 
 		// Unregister the temporary default name that we registered earlier for this schema instance
-		IAssetTools& AssetTools = FModuleManager::GetModuleChecked<FAssetToolsModule>("AssetTools").Get();
 		AssetTools.UnregisterDefaultAssetNameForClass(UCineAssembly::StaticClass());
 	}
 	else
@@ -59,6 +74,10 @@ UObject* UCineAssemblyFactory::FactoryCreateNew(UClass* Class, UObject* InParent
 
 	CineAssemblyToConfigure->CreateSubAssemblies();
 
+	// Refresh the content browser to make sure the original asset tile is not visible (since it may have been moved to the schema's DefaultAssemblyPath)
+	IContentBrowserSingleton& ContentBrowser = FModuleManager::LoadModuleChecked<FContentBrowserModule>("ContentBrowser").Get();
+	ContentBrowser.SetSelectedPaths({ ContentBrowser.GetCurrentPath().GetInternalPathString() }, true);
+
 	return CineAssemblyToConfigure;
 }
 
@@ -85,6 +104,21 @@ bool UCineAssemblyFactory::ConfigureProperties()
 	// This workaround allows us to temporarily set a default name for the assembly asset. 
 	// After FactoryCreateNew() is finished, we will unregister this default name override. 
 	if (!CineAssemblyToConfigure->AssemblyName.Template.IsEmpty())
+	{
+		const FString ResolvedString = ResolveTokenString(CineAssemblyToConfigure->AssemblyName.Template);
+		CineAssemblyToConfigure->AssemblyName.Resolved = FText::FromString(ResolvedString);
+
+		IAssetTools& AssetTools = FModuleManager::GetModuleChecked<FAssetToolsModule>("AssetTools").Get();
+		AssetTools.RegisterDefaultAssetNameForClass(UCineAssembly::StaticClass(), ResolvedString);
+	}
+
+	return true;
+}
+
+FString UCineAssemblyFactory::ResolveTokenString(const FString& TemplateString)
+{
+	FString ResolvedString;
+	if (CineAssemblyToConfigure)
 	{
 		// Add the cine assembly being created to the naming token context
 		UCineAssemblyNamingTokensContext* NamingTokenContext = NewObject<UCineAssemblyNamingTokensContext>();
@@ -94,15 +128,12 @@ bool UCineAssemblyFactory::ConfigureProperties()
 		FilterArgs.AdditionalNamespacesToInclude.Add(UCineAssemblyNamingTokens::TokenNamespace);
 
 		UNamingTokensEngineSubsystem* NamingTokensSubsystem = GEngine->GetEngineSubsystem<UNamingTokensEngineSubsystem>();
-		FNamingTokenResultData Result = NamingTokensSubsystem->EvaluateTokenString(CineAssemblyToConfigure->AssemblyName.Template, FilterArgs, { NamingTokenContext });
-
-		CineAssemblyToConfigure->AssemblyName.Resolved = Result.EvaluatedText;
+		FNamingTokenResultData Result = NamingTokensSubsystem->EvaluateTokenString(TemplateString, FilterArgs, { NamingTokenContext });
 
-		IAssetTools& AssetTools = FModuleManager::GetModuleChecked<FAssetToolsModule>("AssetTools").Get();
-		AssetTools.RegisterDefaultAssetNameForClass(UCineAssembly::StaticClass(), CineAssemblyToConfigure->AssemblyName.Resolved.ToString());
+		ResolvedString = Result.EvaluatedText.ToString();
 	}
 
-	return true;
+	return ResolvedString;
 }
 
 #undef LOCTEXT_NAMESPACE
diff --git a/Engine/Plugins/VirtualProduction/CinematicAssemblyTools/Source/ShotsEditor/Private/CineAssemblyFactory.h b/Engine/Plugins/VirtualProduction/CinematicAssemblyTools/Source/ShotsEditor/Private/CineAssemblyFactory.h
--- a/Engine/Plugins/VirtualProduction/CinematicAssemblyTools/Source/ShotsEditor/Private/CineAssemblyFactory.h
+++ b/Engine/Plugins/VirtualProduction/CinematicAssemblyTools/Source/ShotsEditor/Private/CineAssemblyFactory.h
@@ -25,6 +25,10 @@ protected:
 	virtual bool ConfigureProperties() override;
 	// End UFactory Interface
 
+private:
+	/** Utility function that evaluates the input string with the naming tokens subsystem and returns the resolved string */
+	FString ResolveTokenString(const FString& TemplateString);
+
 private:
 	/** Transient CineAssembly that will be configured by the user and used as a template for the actual asset created by this factory */
 	UCineAssembly* CineAssemblyToConfigure;
diff --git a/Engine/Plugins/VirtualProduction/CinematicAssemblyTools/Source/ShotsEditor/Private/CineAssemblyNamingTokens.cpp b/Engine/Plugins/VirtualProduction/CinematicAssemblyTools/Source/ShotsEditor/Private/CineAssemblyNamingTokens.cpp
--- a/Engine/Plugins/VirtualProduction/CinematicAssemblyTools/Source/ShotsEditor/Private/CineAssemblyNamingTokens.cpp
+++ b/Engine/Plugins/VirtualProduction/CinematicAssemblyTools/Source/ShotsEditor/Private/CineAssemblyNamingTokens.cpp
@@ -48,7 +48,7 @@ void UCineAssemblyNamingTokens::OnCreateDefaultTokens(TArray<FNamingTokenData>&
 		};
 
 	FNamingTokenData AssemblyNameToken;
-	AssemblyNameToken.TokenKey = TEXT("assembly_name");
+	AssemblyNameToken.TokenKey = TEXT("assembly");
 	AssemblyNameToken.DisplayName = LOCTEXT("AssemblyNameToken", "Assembly Name");
 	AssemblyNameToken.TokenProcessorNative.BindLambda([this, AssemblyNameTokenFunc]() { return ExecuteTokenFunc(AssemblyNameTokenFunc); });
 
diff --git a/Engine/Plugins/VirtualProduction/CinematicAssemblyTools/Source/ShotsEditor/Private/CineAssemblySchema.h b/Engine/Plugins/VirtualProduction/CinematicAssemblyTools/Source/ShotsEditor/Private/CineAssemblySchema.h
--- a/Engine/Plugins/VirtualProduction/CinematicAssemblyTools/Source/ShotsEditor/Private/CineAssemblySchema.h
+++ b/Engine/Plugins/VirtualProduction/CinematicAssemblyTools/Source/ShotsEditor/Private/CineAssemblySchema.h
@@ -87,6 +87,13 @@ public:
 	UPROPERTY(EditAnywhere, Category = "Default")
 	FString DefaultAssemblyName;
 
+	/** 
+	 * The default path to use when creating assemblies from this schema
+	 * When an assembly asset is created, this path will be appended to path where the asset would normally have been created.
+	 */
+	UPROPERTY(EditAnywhere, Category = "Default")
+	FString DefaultAssemblyPath;
+
 	/** The thumbnail image to use for this schema and assemblies built from this schema */
 	UPROPERTY(EditAnywhere, Category = "Default")
 	TObjectPtr<UTexture2D> ThumbnailImage;
diff --git a/Engine/Plugins/VirtualProduction/CinematicAssemblyTools/Source/ShotsEditor/Private/ProductionFunctionLibrary.cpp b/Engine/Plugins/VirtualProduction/CinematicAssemblyTools/Source/ShotsEditor/Private/ProductionFunctionLibrary.cpp
--- a/Engine/Plugins/VirtualProduction/CinematicAssemblyTools/Source/ShotsEditor/Private/ProductionFunctionLibrary.cpp
+++ b/Engine/Plugins/VirtualProduction/CinematicAssemblyTools/Source/ShotsEditor/Private/ProductionFunctionLibrary.cpp
@@ -196,11 +196,22 @@ void UProductionFunctionLibrary::RenameAssembly(UCineAssembly* Assembly)
 
 	Assembly->AssemblyName.Resolved = Result.EvaluatedText;
 
+	// Resolve the default assembly path from teh schema, which will be added to the new desired package path
+	FString DefaultAssemblyPath;
+	if (const UCineAssemblySchema* Schema = Assembly->GetSchema())
+	{
+		if (!Schema->DefaultAssemblyPath.IsEmpty())
+		{
+			Result = NamingTokensSubsystem->EvaluateTokenString(Schema->DefaultAssemblyPath, FilterArgs, { NamingTokenContext });
+			DefaultAssemblyPath = Result.EvaluatedText.ToString();
+		}
+	}
+
 	// We now need to rename the asset (which was originally created with a default name) to the resolved default assembly name
 	const FSoftObjectPath OldObjectPath = FSoftObjectPath(Assembly);
 
 	// Ensure that the resolved assembly name is actually unique
-	const FString DesiredPackageName = FPaths::GetPath(Assembly->GetPathName()) / Assembly->AssemblyName.Resolved.ToString();
+	const FString DesiredPackageName = FPaths::GetPath(Assembly->GetPathName()) / DefaultAssemblyPath / Assembly->AssemblyName.Resolved.ToString();
 
 	FString UniquePackageName;
 	FString UniqueAssetName;
diff --git a/Engine/Plugins/VirtualProduction/CinematicAssemblyTools/Source/ShotsEditor/Private/TakeRecorder/CineAssemblyTakeRecorderIntegration.cpp b/Engine/Plugins/VirtualProduction/CinematicAssemblyTools/Source/ShotsEditor/Private/TakeRecorder/CineAssemblyTakeRecorderIntegration.cpp
--- a/Engine/Plugins/VirtualProduction/CinematicAssemblyTools/Source/ShotsEditor/Private/TakeRecorder/CineAssemblyTakeRecorderIntegration.cpp
+++ b/Engine/Plugins/VirtualProduction/CinematicAssemblyTools/Source/ShotsEditor/Private/TakeRecorder/CineAssemblyTakeRecorderIntegration.cpp
@@ -35,6 +35,12 @@ void FCineAssemblyTakeRecorderIntegration::OnRecordingInitialized(UTakeRecorder*
 
 		if (UCineAssembly* Assembly = Cast<UCineAssembly>(TakeRecorder->GetSequence()))
 		{
+			FString AssemblyName;
+			Assembly->GetName(AssemblyName);
+
+			Assembly->AssemblyName.Template = AssemblyName;
+			Assembly->AssemblyName.Resolved = FText::FromString(AssemblyName);
+
 			Assembly->Level = FSoftObjectPath(TakesUtils::DiscoverSourceWorld());
 
 			// Reset the assembly's playback range, since it was initialized to some default value when it was created, but we do not know what its final range will be until we finish recording.
diff --git a/Engine/Plugins/VirtualProduction/CinematicAssemblyTools/Source/ShotsEditor/Private/UI/CineAssembly/SCineAssemblySchemaWindow.cpp b/Engine/Plugins/VirtualProduction/CinematicAssemblyTools/Source/ShotsEditor/Private/UI/CineAssembly/SCineAssemblySchemaWindow.cpp
--- a/Engine/Plugins/VirtualProduction/CinematicAssemblyTools/Source/ShotsEditor/Private/UI/CineAssembly/SCineAssemblySchemaWindow.cpp
+++ b/Engine/Plugins/VirtualProduction/CinematicAssemblyTools/Source/ShotsEditor/Private/UI/CineAssembly/SCineAssemblySchemaWindow.cpp
@@ -795,7 +795,7 @@ TSharedPtr<SWidget> SCineAssemblySchemaWindow::MakeContentTreeContextMenu()
 			);
 		}
 
-		if (SelectedTreeItem != RootItem)
+		if ((SelectedTreeItem != RootItem) && (SelectedTreeItem != TopLevelAssemblyItem))
 		{
 			MenuBuilder.AddMenuEntry(
 				LOCTEXT("RenameAction", "Rename"),
@@ -836,7 +836,7 @@ FReply SCineAssemblySchemaWindow::OnTreeViewKeyDown(const FGeometry& MyGeometry,
 		if (SelectedItems.Num() == 1)
 		{
 			const TSharedPtr<FSchemaTreeItem>& SelectedItem = SelectedItems[0];
-			if (SelectedItem != RootItem)
+			if ((SelectedItem != RootItem) && (SelectedItem != TopLevelAssemblyItem))
 			{
 				DeleteTreeItem(SelectedItem);
 			}
@@ -887,6 +887,18 @@ void SCineAssemblySchemaWindow::InitializeContentTree()
 
 	AddItemsToTree(SchemaToConfigure->FoldersToCreate, FSchemaTreeItem::EItemType::Folder);
 	AddItemsToTree(SchemaToConfigure->SubsequencesToCreate, FSchemaTreeItem::EItemType::Asset);
+
+	// Add the top-level assembly node
+	TSharedPtr<FSchemaTreeItem> ParentItem = FindItemAtPathRecursive(RootItem, SchemaToConfigure->DefaultAssemblyPath);
+	if (ParentItem)
+	{
+		TopLevelAssemblyItem = MakeShared<FSchemaTreeItem>();
+		TopLevelAssemblyItem->Type = FSchemaTreeItem::EItemType::Asset;
+		TopLevelAssemblyItem->Path = TEXT("{assembly}");
+		TopLevelAssemblyItem->Parent = ParentItem;
+
+		ParentItem->ChildAssets.Add(TopLevelAssemblyItem);
+	}
 }
 
 void SCineAssemblySchemaWindow::ExpandTreeRecursive(TSharedPtr<FSchemaTreeItem> TreeItem) const
@@ -912,7 +924,10 @@ void SCineAssemblySchemaWindow::GetAssetListRecursive(TSharedPtr<FSchemaTreeItem
 {
 	for (const TSharedPtr<FSchemaTreeItem>& Asset : TreeItem->ChildAssets)
 	{
-		AssetPathList.Add(Asset->Path);
+		if (Asset != TopLevelAssemblyItem)
+		{
+			AssetPathList.Add(Asset->Path);
+		}
 	}
 
 	for (const TSharedPtr<FSchemaTreeItem>& Child : TreeItem->ChildFolders)
@@ -1115,6 +1130,9 @@ void SCineAssemblySchemaWindow::OnSchemaWindowClosed(const TSharedRef<SWindow>&
 		CineAssemblyNamingTokens->AddMetadataToken(MetadataDesc.Key);
 	}
 
+	// Set the path of the top-level assembly
+	SchemaToConfigure->DefaultAssemblyPath = TopLevelAssemblyItem->Parent->Path;
+
 	// Update the schema's folder list
 	const TArray<FString> CachedFolderList = SchemaToConfigure->FoldersToCreate;
 
diff --git a/Engine/Plugins/VirtualProduction/CinematicAssemblyTools/Source/ShotsEditor/Private/UI/CineAssembly/SCineAssemblySchemaWindow.h b/Engine/Plugins/VirtualProduction/CinematicAssemblyTools/Source/ShotsEditor/Private/UI/CineAssembly/SCineAssemblySchemaWindow.h
--- a/Engine/Plugins/VirtualProduction/CinematicAssemblyTools/Source/ShotsEditor/Private/UI/CineAssembly/SCineAssemblySchemaWindow.h
+++ b/Engine/Plugins/VirtualProduction/CinematicAssemblyTools/Source/ShotsEditor/Private/UI/CineAssembly/SCineAssemblySchemaWindow.h
@@ -248,6 +248,9 @@ private:
 	/** The root item in the tree view */
 	TSharedPtr<FSchemaTreeItem> RootItem;
 
+	/** An item representing where the top-level assembly should be created */
+	TSharedPtr<FSchemaTreeItem> TopLevelAssemblyItem;
+
 	/** The tree view of content folders and assets for this schema */
 	TSharedPtr<STreeView<TSharedPtr<FSchemaTreeItem>>> TreeView;
 
