Unshelved from pending changelist '41864509':

ShotManagement: Add CineAssembly API to Get/Set the Level and Get/Set/Append the Note text.
#rnx

diff --git a/Engine/Plugins/VirtualProduction/CinematicAssemblyTools/Source/ShotsEditor/Private/CineAssembly.cpp b/Engine/Plugins/VirtualProduction/CinematicAssemblyTools/Source/ShotsEditor/Private/CineAssembly.cpp
--- a/Engine/Plugins/VirtualProduction/CinematicAssemblyTools/Source/ShotsEditor/Private/CineAssembly.cpp
+++ b/Engine/Plugins/VirtualProduction/CinematicAssemblyTools/Source/ShotsEditor/Private/CineAssembly.cpp
@@ -260,6 +260,36 @@ void UCineAssembly::Serialize(FArchive& Ar)
 	}
 }
 
+TSoftObjectPtr<UWorld> UCineAssembly::GetLevel()
+{
+	if (Level.IsValid())
+	{
+		return TSoftObjectPtr<UWorld>(Level);
+	}
+	return nullptr;
+}
+
+void UCineAssembly::SetLevel(TSoftObjectPtr<UWorld> InLevel)
+{
+	Level = InLevel.ToSoftObjectPath();
+}
+
+FString UCineAssembly::GetNoteText()
+{
+	return AssemblyNote;
+}
+
+void UCineAssembly::SetNoteText(FString InNote)
+{
+	AssemblyNote = InNote;
+}
+
+void UCineAssembly::AppendToNoteText(FString InNote)
+{
+	AssemblyNote.Append(TEXT("\n"));
+	AssemblyNote.Append(InNote);
+}
+
 FString UCineAssembly::GetFullMetadataString() const
 {
 	FString JsonString;
@@ -270,21 +300,21 @@ FString UCineAssembly::GetFullMetadataString() const
 	return JsonString;
 }
 
-void UCineAssembly::AddMetadataAsString(FString InKey, FString InValue)
+void UCineAssembly::SetMetadataAsString(FString InKey, FString InValue)
 {
 	Modify();
 	MetadataJsonObject->SetStringField(InKey, InValue);
 	AddMetadataNamingToken(InKey);
 }
 
-void UCineAssembly::AddMetadataAsBool(FString InKey, bool InValue)
+void UCineAssembly::SetMetadataAsBool(FString InKey, bool InValue)
 {
 	Modify();
 	MetadataJsonObject->SetBoolField(InKey, InValue);
 	AddMetadataNamingToken(InKey);
 }
 
-void UCineAssembly::AddMetadataAsNumber(FString InKey, double InValue)
+void UCineAssembly::SetMetadataAsNumber(FString InKey, double InValue)
 {
 	Modify();
 	MetadataJsonObject->SetNumberField(InKey, InValue);
@@ -340,7 +370,7 @@ void UCineAssembly::UpdateInstanceMetadata()
 				InstanceMetadataKeys.Add(Pair.Key);
 			}
 
-			AddMetadataAsString(Pair.Key.ToString(), Pair.Value);
+			SetMetadataAsString(Pair.Key.ToString(), Pair.Value);
 		}
 	}
 
diff --git a/Engine/Plugins/VirtualProduction/CinematicAssemblyTools/Source/ShotsEditor/Private/CineAssembly.h b/Engine/Plugins/VirtualProduction/CinematicAssemblyTools/Source/ShotsEditor/Private/CineAssembly.h
--- a/Engine/Plugins/VirtualProduction/CinematicAssemblyTools/Source/ShotsEditor/Private/CineAssembly.h
+++ b/Engine/Plugins/VirtualProduction/CinematicAssemblyTools/Source/ShotsEditor/Private/CineAssembly.h
@@ -41,21 +41,41 @@ public:
 	/** Creates one or more subsequence assets, parented to this assembly, based on the schema */
 	void CreateSubAssemblies();
 
+	/** Get the target level associated with this assembly */
+	UFUNCTION(BlueprintPure, Category = "Shot Management")
+	TSoftObjectPtr<UWorld> GetLevel();
+
+	/** Set the target level associated with this assembly */
+	UFUNCTION(BlueprintCallable, Category = "Shot Management")
+	void SetLevel(TSoftObjectPtr<UWorld> InLevel);
+
+	/** Get the note text associated with this assembly */
+	UFUNCTION(BlueprintPure, Category = "Shot Management")
+	FString GetNoteText();
+
+	/** Set the note text associated with this assembly */
+	UFUNCTION(BlueprintCallable, Category = "Shot Management")
+	void SetNoteText(FString InNote);
+
+	/** Append to the note text associated with this assembly */
+	UFUNCTION(BlueprintCallable, Category = "Shot Management")
+	void AppendToNoteText(FString InNote);
+
 	/** Get all of the metadata for this assembly as a formatted JSON string */
 	UFUNCTION(BlueprintPure, Category = "Shot Management")
 	FString GetFullMetadataString() const;
 
 	/** Add a string as metadata to this assembly */
 	UFUNCTION(BlueprintCallable, Category = "Shot Management")
-	void AddMetadataAsString(FString InKey, FString InValue);
+	void SetMetadataAsString(FString InKey, FString InValue);
 
 	/** Add a boolean as metadata to this assembly */
 	UFUNCTION(BlueprintCallable, Category = "Shot Management")
-	void AddMetadataAsBool(FString InKey, bool InValue);
+	void SetMetadataAsBool(FString InKey, bool InValue);
 
 	/** Add a double-precision floating point number as metadata to this assembly */
 	UFUNCTION(BlueprintCallable, Category = "Shot Management")
-	void AddMetadataAsNumber(FString InKey, double InValue);
+	void SetMetadataAsNumber(FString InKey, double InValue);
 
 	/** Get the metadata value for the input key as a string (if it exists) */
 	UFUNCTION(BlueprintPure, Category = "Shot Management", meta = (ReturnDisplayName = "FoundKey"))
diff --git a/Engine/Plugins/VirtualProduction/CinematicAssemblyTools/Source/ShotsEditor/Private/CineAssemblyCustomization.cpp b/Engine/Plugins/VirtualProduction/CinematicAssemblyTools/Source/ShotsEditor/Private/CineAssemblyCustomization.cpp
--- a/Engine/Plugins/VirtualProduction/CinematicAssemblyTools/Source/ShotsEditor/Private/CineAssemblyCustomization.cpp
+++ b/Engine/Plugins/VirtualProduction/CinematicAssemblyTools/Source/ShotsEditor/Private/CineAssemblyCustomization.cpp
@@ -79,7 +79,7 @@ void FCineAssemblyCustomization::CustomizeDetails(IDetailLayoutBuilder& DetailBu
 								})
 							.OnTextCommitted_Lambda([CustomizedCineAssembly, MetadataDesc](const FText& InText, ETextCommit::Type InCommitType)
 								{
-									CustomizedCineAssembly->AddMetadataAsString(MetadataDesc.Key, InText.ToString());
+									CustomizedCineAssembly->SetMetadataAsString(MetadataDesc.Key, InText.ToString());
 								})
 					];
 			}
@@ -98,7 +98,7 @@ void FCineAssemblyCustomization::CustomizeDetails(IDetailLayoutBuilder& DetailBu
 					.OnCheckStateChanged_Lambda([CustomizedCineAssembly, MetadataDesc](ECheckBoxState CheckBoxState)
 						{
 							const bool Value = (CheckBoxState == ECheckBoxState::Checked);
-							CustomizedCineAssembly->AddMetadataAsBool(MetadataDesc.Key, Value);
+							CustomizedCineAssembly->SetMetadataAsBool(MetadataDesc.Key, Value);
 						});
 			}
 			else if (MetadataDesc.Type == ECineAssemblyMetadataType::Number)
@@ -116,7 +116,7 @@ void FCineAssemblyCustomization::CustomizeDetails(IDetailLayoutBuilder& DetailBu
 						})
 					.OnValueChanged_Lambda([CustomizedCineAssembly, MetadataDesc](double InValue)
 						{
-							CustomizedCineAssembly->AddMetadataAsNumber(MetadataDesc.Key, InValue);
+							CustomizedCineAssembly->SetMetadataAsNumber(MetadataDesc.Key, InValue);
 						});
 			}
 			else if (MetadataDesc.Type == ECineAssemblyMetadataType::AssetPath)
@@ -137,7 +137,7 @@ void FCineAssemblyCustomization::CustomizeDetails(IDetailLayoutBuilder& DetailBu
 					.OnObjectChanged_Lambda([CustomizedCineAssembly, MetadataDesc](const FAssetData& InAssetData)
 						{
 							// Store the path of the selected object as a string
-							CustomizedCineAssembly->AddMetadataAsString(MetadataDesc.Key, InAssetData.GetObjectPathString());
+							CustomizedCineAssembly->SetMetadataAsString(MetadataDesc.Key, InAssetData.GetObjectPathString());
 						});
 			}
 			else if (MetadataDesc.Type == ECineAssemblyMetadataType::CineAssembly)
@@ -167,7 +167,7 @@ void FCineAssemblyCustomization::CustomizeDetails(IDetailLayoutBuilder& DetailBu
 					.OnObjectChanged_Lambda([CustomizedCineAssembly, MetadataDesc](const FAssetData& InAssetData)
 						{
 							// Store the path of the selected object as a string
-							CustomizedCineAssembly->AddMetadataAsString(MetadataDesc.Key, InAssetData.GetObjectPathString());
+							CustomizedCineAssembly->SetMetadataAsString(MetadataDesc.Key, InAssetData.GetObjectPathString());
 						});
 			}
 			else
diff --git a/Engine/Plugins/VirtualProduction/CinematicAssemblyTools/Source/ShotsEditor/Private/ProductionFunctionLibrary.cpp b/Engine/Plugins/VirtualProduction/CinematicAssemblyTools/Source/ShotsEditor/Private/ProductionFunctionLibrary.cpp
--- a/Engine/Plugins/VirtualProduction/CinematicAssemblyTools/Source/ShotsEditor/Private/ProductionFunctionLibrary.cpp
+++ b/Engine/Plugins/VirtualProduction/CinematicAssemblyTools/Source/ShotsEditor/Private/ProductionFunctionLibrary.cpp
@@ -151,30 +151,30 @@ void UProductionFunctionLibrary::AddMetadataToAssembly(UCineAssembly* Assembly,
 		{
 			if (MetadataDesc->Type == ECineAssemblyMetadataType::String)
 			{
-				Assembly->AddMetadataAsString(Pair.Key, Pair.Value);
+				Assembly->SetMetadataAsString(Pair.Key, Pair.Value);
 			}
 			else if (MetadataDesc->Type == ECineAssemblyMetadataType::Bool)
 			{
-				Assembly->AddMetadataAsBool(Pair.Key, Pair.Value.ToBool());
+				Assembly->SetMetadataAsBool(Pair.Key, Pair.Value.ToBool());
 			}
 			else if (MetadataDesc->Type == ECineAssemblyMetadataType::Number)
 			{
-				Assembly->AddMetadataAsNumber(Pair.Key, FCString::Atof(*Pair.Value));
+				Assembly->SetMetadataAsNumber(Pair.Key, FCString::Atof(*Pair.Value));
 			}
 			else if (MetadataDesc->Type == ECineAssemblyMetadataType::AssetPath)
 			{
-				Assembly->AddMetadataAsString(Pair.Key, Pair.Value);
+				Assembly->SetMetadataAsString(Pair.Key, Pair.Value);
 			}
 			else if (MetadataDesc->Type == ECineAssemblyMetadataType::CineAssembly)
 			{
-				Assembly->AddMetadataAsString(Pair.Key, Pair.Value);
+				Assembly->SetMetadataAsString(Pair.Key, Pair.Value);
 			}
 		}
 		else
 		{
 			// If the metadata string does not match anything specified by the schema, simply add the metadata as a string.
 			Assembly->InstanceMetadata.Add(*Pair.Key, Pair.Value);
-			Assembly->AddMetadataAsString(Pair.Key, Pair.Value);
+			Assembly->SetMetadataAsString(Pair.Key, Pair.Value);
 		}
 	}
 }
