ShotManagement: Add Production and ParentAssembly properties to UCineAssembly, along with detail customizations for both properties.
Add a ParentSchema property to UCineAssemblySchema which restricts the assemblies that can be picked for ParentAssembly.

#rnx

diff --git a/Engine/Plugins/VirtualProduction/CinematicAssemblyTools/Source/ShotsEditor/Private/CineAssembly.h b/Engine/Plugins/VirtualProduction/CinematicAssemblyTools/Source/ShotsEditor/Private/CineAssembly.h
--- a/Engine/Plugins/VirtualProduction/CinematicAssemblyTools/Source/ShotsEditor/Private/CineAssembly.h
+++ b/Engine/Plugins/VirtualProduction/CinematicAssemblyTools/Source/ShotsEditor/Private/CineAssembly.h
@@ -105,7 +105,7 @@ public:
 	FTemplateString AssemblyName;
 
 	/** The level to open before opening this asset in Sequencer */
-	UPROPERTY(EditAnywhere, Category = "TargetLevel", meta = (AllowedClasses = "/Script/Engine.World"))
+	UPROPERTY(EditAnywhere, Category = "Default", meta = (AllowedClasses = "/Script/Engine.World"))
 	FSoftObjectPath Level;
 
 	/** User added metadata key/value pairs, which will be added as additional asset registry tags */
@@ -116,6 +116,14 @@ public:
 	UPROPERTY()
 	FString AssemblyNote;
 
+	/** Reference to another assembly asset that is the parent of this assembly */
+	UPROPERTY(EditAnywhere, Category = "Default", meta = (AllowedClasses = "/Script/ShotsEditor.CineAssembly"))
+	FSoftObjectPath ParentAssembly;
+
+	/** The ID of the Cinematic Production that this assembly is associated with */
+	UPROPERTY(EditAnywhere, Category = "Default")
+	FGuid Production;
+
 	/** Array of template names (possibly containing tokens), based on the schema, used to create the SubAssemblies */
 	UPROPERTY()
 	TArray<FTemplateString> SubAssemblyNames;
diff --git a/Engine/Plugins/VirtualProduction/CinematicAssemblyTools/Source/ShotsEditor/Private/CineAssemblyCustomization.cpp b/Engine/Plugins/VirtualProduction/CinematicAssemblyTools/Source/ShotsEditor/Private/CineAssemblyCustomization.cpp
--- a/Engine/Plugins/VirtualProduction/CinematicAssemblyTools/Source/ShotsEditor/Private/CineAssemblyCustomization.cpp
+++ b/Engine/Plugins/VirtualProduction/CinematicAssemblyTools/Source/ShotsEditor/Private/CineAssemblyCustomization.cpp
@@ -9,10 +9,13 @@
 #include "CineAssemblyNamingTokens.h"
 #include "DetailLayoutBuilder.h"
 #include "DetailWidgetRow.h"
+#include "Framework/MultiBox/MultiBoxBuilder.h"
 #include "NamingTokensEngineSubsystem.h"
+#include "ProductionSettings.h"
 #include "PropertyCustomizationHelpers.h"
 #include "STemplateStringEditableTextBox.h"
 #include "Widgets/Input/SCheckBox.h"
+#include "Widgets/Input/SComboButton.h"
 #include "Widgets/Input/SMultiLineEditableTextBox.h"
 #include "Widgets/Input/SNumericEntryBox.h"
 #include "Widgets/Layout/SBox.h"
@@ -36,201 +39,287 @@ void FCineAssemblyCustomization::CustomizeDetails(IDetailLayoutBuilder& DetailBu
 		return;
 	}
 
-	UCineAssembly* CustomizedCineAssembly = Cast<UCineAssembly>(CustomizedObjects[0]);
+	CustomizedCineAssembly = Cast<UCineAssembly>(CustomizedObjects[0]);
 
-	if (const UCineAssemblySchema* BaseSchema = CustomizedCineAssembly->GetSchema())
+	CustomizeDefaultCategory(DetailBuilder);
+	CustomizeMetadataCategory(DetailBuilder);
+	CustomizeSubsequenceCategory(DetailBuilder);
+}
+
+void FCineAssemblyCustomization::CustomizeDefaultCategory(IDetailLayoutBuilder& DetailBuilder)
+{
+	IDetailCategoryBuilder& DefaultCategory = DetailBuilder.EditCategory("Default", LOCTEXT("DefaultCategoryName", "Default"));
+
+	// Customize the ParenteAssembly and Production properties of the UCineAssembly
+	TSharedRef<IPropertyHandle> LevelPropertyHandle = DetailBuilder.GetProperty(GET_MEMBER_NAME_CHECKED(UCineAssembly, Level));
+	DefaultCategory.AddProperty(LevelPropertyHandle);
+
+	TSharedRef<IPropertyHandle> ParentPropertyHandle = DetailBuilder.GetProperty(GET_MEMBER_NAME_CHECKED(UCineAssembly, ParentAssembly));
+	IDetailPropertyRow& ParentPropertyRow = DefaultCategory.AddProperty(ParentPropertyHandle);
+
+	TSharedRef<IPropertyHandle> ProdctionPropertyHandle = DetailBuilder.GetProperty(GET_MEMBER_NAME_CHECKED(UCineAssembly, Production));
+	IDetailPropertyRow& ProductionPropertyRow = DefaultCategory.AddProperty(ProdctionPropertyHandle);
+
+	const UCineAssemblySchema* BaseSchema = CustomizedCineAssembly->GetSchema();
+	const FSoftObjectPath ParentSchema = BaseSchema ? BaseSchema->ParentSchema : FSoftObjectPath();
+
+	ParentPropertyRow.CustomWidget()
+		.NameContent()
+		[
+			ParentPropertyHandle->CreatePropertyNameWidget()
+		]
+		.ValueContent()
+		[
+			SNew(SObjectPropertyEntryBox)
+				.AllowedClass(UCineAssembly::StaticClass())
+				.ThumbnailPool(DetailBuilder.GetThumbnailPool())
+				.AllowCreate(true)
+				.OnShouldFilterAsset_Raw(this, &FCineAssemblyCustomization::ShouldFilterAssetBySchema, ParentSchema)
+				.ObjectPath_Lambda([this]()
+					{
+						return CustomizedCineAssembly->ParentAssembly.ToString();
+					})
+				.OnObjectChanged_Lambda([this](const FAssetData& InAssetData)
+					{
+						CustomizedCineAssembly->Modify();
+						CustomizedCineAssembly->ParentAssembly = InAssetData.GetObjectPathString();
+					})
+		];
+
+	ProductionPropertyRow.CustomWidget()
+		.NameContent()
+		[
+			ProdctionPropertyHandle->CreatePropertyNameWidget()
+		]
+		.ValueContent()
+		[
+			SNew(SComboButton)
+				.VAlign(VAlign_Center)
+				.OnGetMenuContent(this, &FCineAssemblyCustomization::BuildProductionNameMenu)
+				.ButtonContent()
+				[
+					SNew(STextBlock)
+						.Text_Lambda([this]()
+							{
+								const UProductionSettings* const ProductionSettings = GetDefault<UProductionSettings>();
+								TOptional<const FCinematicProduction> Production = ProductionSettings->GetProduction(CustomizedCineAssembly->Production);
+								return Production.IsSet() ? FText::FromString(Production.GetValue().ProductionName) : FText::FromName(NAME_None);
+							})
+						.Font(IDetailLayoutBuilder::GetDetailFont())
+				]
+		];
+}
+
+void FCineAssemblyCustomization::CustomizeMetadataCategory(IDetailLayoutBuilder& DetailBuilder)
+{
+	const UCineAssemblySchema* BaseSchema = CustomizedCineAssembly->GetSchema();
+	if (!BaseSchema)
 	{
-		// 'Edit' the target level category to get it to appear first in the details view
-		DetailBuilder.EditCategory("TargetLevel", LOCTEXT("TargetLevelCategoryName", "Target Level"));
+		return;
+	}
 
-		// Add a new category for Schema Metadata properties
-		const FText MetadataCategoryName = FText::Format(LOCTEXT("SchemaMetadataCategoryName", "{0} Metadata"), FText::FromString(BaseSchema->SchemaName));
-		IDetailCategoryBuilder& MetadataCategory = DetailBuilder.EditCategory("SchemaMetadata", MetadataCategoryName);
+	// Add a new category for Schema Metadata properties
+	const FText MetadataCategoryName = FText::Format(LOCTEXT("SchemaMetadataCategoryName", "{0} Metadata"), FText::FromString(BaseSchema->SchemaName));
+	IDetailCategoryBuilder& MetadataCategory = DetailBuilder.EditCategory("SchemaMetadata", MetadataCategoryName);
 
-		// Add a new category for Schema Subsequences
-		const FText SubsequencesCategoryName = LOCTEXT("SchemaSubseuqnecesCategoryName", "Subsequences");
-		IDetailCategoryBuilder& SubsequenceCategory = DetailBuilder.EditCategory("SchemaSubsequences", SubsequencesCategoryName);
+	// Add a property row for each metadata struct in the customized Assembly's base schema
+	for (const FAssemblyMetadataDesc& MetadataDesc : BaseSchema->AssemblyMetadata)
+	{
+		if (MetadataDesc.Key.IsEmpty())
+		{
+			continue;
+		}
 
-		// Add a property row for each metadata struct in the customized Assembly's base schema
-		for (const FAssemblyMetadataDesc& MetadataDesc : BaseSchema->AssemblyMetadata)
+		TSharedPtr<SWidget> ValueWidget;
+		if (MetadataDesc.Type == ECineAssemblyMetadataType::String)
 		{
-			if (MetadataDesc.Key.IsEmpty())
-			{
-				continue;
-			}
-
-			TSharedPtr<SWidget> ValueWidget;
-			if (MetadataDesc.Type == ECineAssemblyMetadataType::String)
-			{
-				ValueWidget = SNew(SBox)
-					.MaxDesiredHeight(120.0f)
-					[
-						SNew(SMultiLineEditableTextBox)
-							.Font(FAppStyle::GetFontStyle(TEXT("PropertyWindow.NormalFont")))
-							.AutoWrapText(true)
-							.Text_Lambda([CustomizedCineAssembly, MetadataDesc]()
-								{
-									FString Value;
-									if (!CustomizedCineAssembly->GetMetadataAsString(MetadataDesc.Key, Value))
-									{
-										Value = MetadataDesc.DefaultValue.Get<FString>();
-									}
-									return FText::FromString(Value);
-								})
-							.OnTextCommitted_Lambda([CustomizedCineAssembly, MetadataDesc](const FText& InText, ETextCommit::Type InCommitType)
-								{
-									CustomizedCineAssembly->SetMetadataAsString(MetadataDesc.Key, InText.ToString());
-								})
-					];
-			}
-			else if (MetadataDesc.Type == ECineAssemblyMetadataType::Bool)
-			{
-				ValueWidget = SNew(SCheckBox)
-					.IsChecked_Lambda([CustomizedCineAssembly, MetadataDesc]()
-						{
-							bool Value;
-							if (!CustomizedCineAssembly->GetMetadataAsBool(MetadataDesc.Key, Value))
-							{
-								Value = MetadataDesc.DefaultValue.Get<bool>();
-							}
-							return Value ? ECheckBoxState::Checked : ECheckBoxState::Unchecked;
-						})
-					.OnCheckStateChanged_Lambda([CustomizedCineAssembly, MetadataDesc](ECheckBoxState CheckBoxState)
-						{
-							const bool Value = (CheckBoxState == ECheckBoxState::Checked);
-							CustomizedCineAssembly->SetMetadataAsBool(MetadataDesc.Key, Value);
-						});
-			}
-			else if (MetadataDesc.Type == ECineAssemblyMetadataType::Number)
-			{
-				ValueWidget = SNew(SNumericEntryBox<double>)
-					.Font(FAppStyle::GetFontStyle(TEXT("PropertyWindow.NormalFont")))
-					.Value_Lambda([CustomizedCineAssembly, MetadataDesc]()
-						{
-							double Value;
-							if (!CustomizedCineAssembly->GetMetadataAsNumber(MetadataDesc.Key, Value))
+			ValueWidget = SNew(SBox)
+				.MaxDesiredHeight(120.0f)
+				[
+					SNew(SMultiLineEditableTextBox)
+						.Font(FAppStyle::GetFontStyle(TEXT("PropertyWindow.NormalFont")))
+						.AutoWrapText(true)
+						.Text_Lambda([this, MetadataDesc]()
 							{
-								Value = MetadataDesc.DefaultValue.Get<double>();
-							}
-							return Value;
-						})
-					.OnValueChanged_Lambda([CustomizedCineAssembly, MetadataDesc](double InValue)
-						{
-							CustomizedCineAssembly->SetMetadataAsNumber(MetadataDesc.Key, InValue);
-						});
-			}
-			else if (MetadataDesc.Type == ECineAssemblyMetadataType::AssetPath)
-			{
-				ValueWidget = SNew(SObjectPropertyEntryBox)
-					.AllowedClass(MetadataDesc.AssetClass.ResolveClass())
-					.ThumbnailPool(DetailBuilder.GetThumbnailPool())
-					.AllowCreate(true)
-					.ObjectPath_Lambda([CustomizedCineAssembly, MetadataDesc]()
-						{
-							FString Value;
-							if (!CustomizedCineAssembly->GetMetadataAsString(MetadataDesc.Key, Value))
+								FString Value;
+								if (!CustomizedCineAssembly->GetMetadataAsString(MetadataDesc.Key, Value))
+								{
+									Value = MetadataDesc.DefaultValue.Get<FString>();
+								}
+								return FText::FromString(Value);
+							})
+						.OnTextCommitted_Lambda([this, MetadataDesc](const FText& InText, ETextCommit::Type InCommitType)
 							{
-								Value = MetadataDesc.DefaultValue.Get<FString>();
-							}
-							return Value;
-						})
-					.OnObjectChanged_Lambda([CustomizedCineAssembly, MetadataDesc](const FAssetData& InAssetData)
+								CustomizedCineAssembly->SetMetadataAsString(MetadataDesc.Key, InText.ToString());
+							})
+				];
+		}
+		else if (MetadataDesc.Type == ECineAssemblyMetadataType::Bool)
+		{
+			ValueWidget = SNew(SCheckBox)
+				.IsChecked_Lambda([this, MetadataDesc]()
+					{
+						bool Value;
+						if (!CustomizedCineAssembly->GetMetadataAsBool(MetadataDesc.Key, Value))
 						{
-							// Store the path of the selected object as a string
-							CustomizedCineAssembly->SetMetadataAsString(MetadataDesc.Key, InAssetData.GetObjectPathString());
-						});
-			}
-			else if (MetadataDesc.Type == ECineAssemblyMetadataType::CineAssembly)
-			{
-				ValueWidget = SNew(SObjectPropertyEntryBox)
-					.AllowedClass(UCineAssembly::StaticClass())
-					.ThumbnailPool(DetailBuilder.GetThumbnailPool())
-					.AllowCreate(true)
-					.OnShouldFilterAsset_Lambda([MetadataDesc](const FAssetData& InAssetData) -> bool
+							Value = MetadataDesc.DefaultValue.Get<bool>();
+						}
+						return Value ? ECheckBoxState::Checked : ECheckBoxState::Unchecked;
+					})
+				.OnCheckStateChanged_Lambda([this, MetadataDesc](ECheckBoxState CheckBoxState)
+					{
+						const bool Value = (CheckBoxState == ECheckBoxState::Checked);
+						CustomizedCineAssembly->SetMetadataAsBool(MetadataDesc.Key, Value);
+					});
+		}
+		else if (MetadataDesc.Type == ECineAssemblyMetadataType::Number)
+		{
+			ValueWidget = SNew(SNumericEntryBox<double>)
+				.Font(FAppStyle::GetFontStyle(TEXT("PropertyWindow.NormalFont")))
+				.Value_Lambda([this, MetadataDesc]()
+					{
+						double Value;
+						if (!CustomizedCineAssembly->GetMetadataAsNumber(MetadataDesc.Key, Value))
 						{
-							const FAssetDataTagMapSharedView::FFindTagResult AssemblyType = InAssetData.TagsAndValues.FindTag(UCineAssembly::AssetRegistryTag_AssemblyType);
-							if (AssemblyType.IsSet())
-							{
-								return !AssemblyType.GetValue().Equals(MetadataDesc.SchemaType.GetAssetName());
-							}
-							return true;
-						})
-					.ObjectPath_Lambda([CustomizedCineAssembly, MetadataDesc]()
+							Value = MetadataDesc.DefaultValue.Get<double>();
+						}
+						return Value;
+					})
+				.OnValueChanged_Lambda([this, MetadataDesc](double InValue)
+					{
+						CustomizedCineAssembly->SetMetadataAsNumber(MetadataDesc.Key, InValue);
+					});
+		}
+		else if (MetadataDesc.Type == ECineAssemblyMetadataType::AssetPath)
+		{
+			ValueWidget = SNew(SObjectPropertyEntryBox)
+				.AllowedClass(MetadataDesc.AssetClass.ResolveClass())
+				.ThumbnailPool(DetailBuilder.GetThumbnailPool())
+				.AllowCreate(true)
+				.ObjectPath_Lambda([this, MetadataDesc]()
+					{
+						FString Value;
+						if (!CustomizedCineAssembly->GetMetadataAsString(MetadataDesc.Key, Value))
 						{
-							FString Value;
-							if (!CustomizedCineAssembly->GetMetadataAsString(MetadataDesc.Key, Value))
-							{
-								Value = MetadataDesc.DefaultValue.Get<FString>();
-							}
-							return Value;
-						})
-					.OnObjectChanged_Lambda([CustomizedCineAssembly, MetadataDesc](const FAssetData& InAssetData)
+							Value = MetadataDesc.DefaultValue.Get<FString>();
+						}
+						return Value;
+					})
+				.OnObjectChanged_Lambda([this, MetadataDesc](const FAssetData& InAssetData)
+					{
+						// Store the path of the selected object as a string
+						CustomizedCineAssembly->SetMetadataAsString(MetadataDesc.Key, InAssetData.GetObjectPathString());
+					});
+		}
+		else if (MetadataDesc.Type == ECineAssemblyMetadataType::CineAssembly)
+		{
+			ValueWidget = SNew(SObjectPropertyEntryBox)
+				.AllowedClass(UCineAssembly::StaticClass())
+				.ThumbnailPool(DetailBuilder.GetThumbnailPool())
+				.AllowCreate(true)
+				.OnShouldFilterAsset_Raw(this, &FCineAssemblyCustomization::ShouldFilterAssetBySchema, MetadataDesc.SchemaType)
+				.ObjectPath_Lambda([this, MetadataDesc]()
+					{
+						FString Value;
+						if (!CustomizedCineAssembly->GetMetadataAsString(MetadataDesc.Key, Value))
 						{
-							// Store the path of the selected object as a string
-							CustomizedCineAssembly->SetMetadataAsString(MetadataDesc.Key, InAssetData.GetObjectPathString());
-						});
-			}
-			else
-			{
-				checkNoEntry();
-			}
-
-			MetadataCategory.AddCustomRow(FText::FromString(MetadataDesc.Key))
-				.NameContent()
-				[
-					SNew(STextBlock)
-						.Text(FText::FromString(MetadataDesc.Key))
-						.Font(IDetailLayoutBuilder::GetDetailFont())
-				]
-				.ValueContent()
-				.HAlign(HAlign_Fill)
-				[
-					ValueWidget.ToSharedRef()
-				];
+							Value = MetadataDesc.DefaultValue.Get<FString>();
+						}
+						return Value;
+					})
+				.OnObjectChanged_Lambda([this, MetadataDesc](const FAssetData& InAssetData)
+					{
+						// Store the path of the selected object as a string
+						CustomizedCineAssembly->SetMetadataAsString(MetadataDesc.Key, InAssetData.GetObjectPathString());
+					});
 		}
-
-		SubAssemblyNames.Reset();
-		for (const FString& SubsequenceName : BaseSchema->SubsequencesToCreate)
+		else
 		{
-			FTemplateString SubsequenceTemplateString;
-			SubsequenceTemplateString.Template = SubsequenceName;
-			SubAssemblyNames.Add(SubsequenceTemplateString);
+			checkNoEntry();
 		}
 
-		for (int32 Index = 0; Index < SubAssemblyNames.Num(); ++Index)
+		MetadataCategory.AddCustomRow(FText::FromString(MetadataDesc.Key))
+			.NameContent()
+			[
+				SNew(STextBlock)
+					.Text(FText::FromString(MetadataDesc.Key))
+					.Font(IDetailLayoutBuilder::GetDetailFont())
+			]
+			.ValueContent()
+			.HAlign(HAlign_Fill)
+			[
+				ValueWidget.ToSharedRef()
+			];
+	}
+}
+
+void FCineAssemblyCustomization::CustomizeSubsequenceCategory(IDetailLayoutBuilder& DetailBuilder)
+{
+	const UCineAssemblySchema* BaseSchema = CustomizedCineAssembly->GetSchema();
+	if (!BaseSchema)
+	{
+		return;
+	}
+
+	// Add a new category for Schema Subsequences
+	const FText SubsequencesCategoryName = LOCTEXT("SchemaSubseuqnecesCategoryName", "Subsequences");
+	IDetailCategoryBuilder& SubsequenceCategory = DetailBuilder.EditCategory("SchemaSubsequences", SubsequencesCategoryName);
+
+	SubAssemblyNames.Reset();
+	for (const FString& SubsequenceName : BaseSchema->SubsequencesToCreate)
+	{
+		FTemplateString SubsequenceTemplateString;
+		SubsequenceTemplateString.Template = SubsequenceName;
+		SubAssemblyNames.Add(SubsequenceTemplateString);
+	}
+
+	for (int32 Index = 0; Index < SubAssemblyNames.Num(); ++Index)
+	{
+		SubsequenceCategory.AddCustomRow(FText::GetEmpty())
+			.NameContent()
+			.HAlign(HAlign_Right)
+			[
+				SNew(SCheckBox)
+					.IsChecked_Raw(this, &FCineAssemblyCustomization::IsSubAssemblyChecked, Index)
+					.OnCheckStateChanged_Raw(this, &FCineAssemblyCustomization::SubAssemblyCheckStateChanged, Index)
+			]
+			.ValueContent()
+			.HAlign(HAlign_Fill)
+			[
+				SNew(STemplateStringEditableTextBox)
+					.Text_Raw(this, &FCineAssemblyCustomization::GetTemplateText, Index)
+					.ResolvedText_Lambda([this, Index]() { return GetResolvedText(Index); })
+					.OnTextCommitted_Raw(this, &FCineAssemblyCustomization::OnTemplateTextCommitted, Index)
+			];
+	}
+}
+
+bool FCineAssemblyCustomization::ShouldFilterAssetBySchema(const FAssetData& InAssetData, FSoftObjectPath Schema)
+{
+	if (Schema.IsValid())
+	{
+		const FAssetDataTagMapSharedView::FFindTagResult AssemblyType = InAssetData.TagsAndValues.FindTag(UCineAssembly::AssetRegistryTag_AssemblyType);
+		if (AssemblyType.IsSet())
 		{
-			SubsequenceCategory.AddCustomRow(FText::GetEmpty())
-				.NameContent()
-				.HAlign(HAlign_Right)
-				[
-					SNew(SCheckBox)
-						.IsChecked_Raw(this, &FCineAssemblyCustomization::IsSubAssemblyChecked, CustomizedCineAssembly, Index)
-						.OnCheckStateChanged_Raw(this, &FCineAssemblyCustomization::SubAssemblyCheckStateChanged, CustomizedCineAssembly, Index)
-				]
-				.ValueContent()
-				.HAlign(HAlign_Fill)
-				[
-					SNew(STemplateStringEditableTextBox)
-						.Text_Raw(this, &FCineAssemblyCustomization::GetTemplateText, Index)
-						.ResolvedText_Lambda([this, Index]() { return GetResolvedText(Index); })
-						.OnTextCommitted_Raw(this, &FCineAssemblyCustomization::OnTemplateTextCommitted, CustomizedCineAssembly, Index)
-				];
+			return !AssemblyType.GetValue().Equals(Schema.GetAssetName());
 		}
+		return true;
 	}
+	return false;
 }
 
-ECheckBoxState FCineAssemblyCustomization::IsSubAssemblyChecked(UCineAssembly* Assembly, int32 Index) const
+ECheckBoxState FCineAssemblyCustomization::IsSubAssemblyChecked(int32 Index) const
 {
 	if (SubAssemblyNames.IsValidIndex(Index))
 	{
 		const FTemplateString& SubAssemblyName = SubAssemblyNames[Index];
-		const bool bResult = Algo::ContainsBy(Assembly->SubAssemblyNames, SubAssemblyName.Template, &FTemplateString::Template);
+		const bool bResult = Algo::ContainsBy(CustomizedCineAssembly->SubAssemblyNames, SubAssemblyName.Template, &FTemplateString::Template);
 		return bResult ? ECheckBoxState::Checked : ECheckBoxState::Unchecked;
 	}
 	return ECheckBoxState::Undetermined;
 }
 
-void FCineAssemblyCustomization::SubAssemblyCheckStateChanged(ECheckBoxState CheckBoxState, UCineAssembly* Assembly, int32 Index)
+void FCineAssemblyCustomization::SubAssemblyCheckStateChanged(ECheckBoxState CheckBoxState, int32 Index)
 {
 	if (SubAssemblyNames.IsValidIndex(Index))
 	{
@@ -238,12 +327,12 @@ void FCineAssemblyCustomization::SubAssemblyCheckStateChanged(ECheckBoxState Che
 
 		if (CheckBoxState == ECheckBoxState::Checked)
 		{
-			Assembly->SubAssemblyNames.Add(SubAssemblyName);
+			CustomizedCineAssembly->SubAssemblyNames.Add(SubAssemblyName);
 		}
 		else if (CheckBoxState == ECheckBoxState::Unchecked)
 		{
-			const int32 IndexToRemove = Algo::IndexOfBy(Assembly->SubAssemblyNames, SubAssemblyName.Template, &FTemplateString::Template);
-			Assembly->SubAssemblyNames.RemoveAt(IndexToRemove);
+			const int32 IndexToRemove = Algo::IndexOfBy(CustomizedCineAssembly->SubAssemblyNames, SubAssemblyName.Template, &FTemplateString::Template);
+			CustomizedCineAssembly->SubAssemblyNames.RemoveAt(IndexToRemove);
 		}
 	}
 }
@@ -270,13 +359,13 @@ FText FCineAssemblyCustomization::GetResolvedText(int32 Index)
 	return FText::GetEmpty();
 }
 
-void FCineAssemblyCustomization::OnTemplateTextCommitted(const FText& InText, ETextCommit::Type InCommitType, UCineAssembly* Assembly, int32 Index)
+void FCineAssemblyCustomization::OnTemplateTextCommitted(const FText& InText, ETextCommit::Type InCommitType, int32 Index)
 {
 	if (SubAssemblyNames.IsValidIndex(Index))
 	{
 		FTemplateString& SubAssemblyName = SubAssemblyNames[Index];
 
-		FTemplateString* TemplateString = Algo::FindBy(Assembly->SubAssemblyNames, SubAssemblyName.Template, &FTemplateString::Template);
+		FTemplateString* TemplateString = Algo::FindBy(CustomizedCineAssembly->SubAssemblyNames, SubAssemblyName.Template, &FTemplateString::Template);
 
 		const FString Path = FPaths::GetPath(SubAssemblyName.Template);
 		SubAssemblyName.Template = Path / InText.ToString();
@@ -311,4 +400,46 @@ void FCineAssemblyCustomization::EvaluateTokenString(FTemplateString& TokenStrin
 	TokenString.Resolved = Result.EvaluatedText;
 }
 
+TSharedRef<SWidget> FCineAssemblyCustomization::BuildProductionNameMenu()
+{
+	FMenuBuilder MenuBuilder(true, nullptr);
+
+	// Always add a "None" option
+	FName NoActiveProductionName = NAME_None;
+	MenuBuilder.AddMenuEntry(
+		FText::FromName(NoActiveProductionName),
+		FText::FromName(NoActiveProductionName),
+		FSlateIcon(),
+		FUIAction(FExecuteAction::CreateLambda([this]() 
+			{ 
+				CustomizedCineAssembly->Modify();
+				CustomizedCineAssembly->Production = FGuid();
+			})),
+		NAME_None,
+		EUserInterfaceActionType::None
+	);
+
+	const UProductionSettings* const ProductionSettings = GetDefault<UProductionSettings>();
+
+	// Add a menu option with the production name for each production available in this project
+	const TArray<FCinematicProduction> Productions = ProductionSettings->GetProductions();
+	for (const FCinematicProduction& Production : Productions)
+	{
+		MenuBuilder.AddMenuEntry(
+			FText::FromString(Production.ProductionName),
+			FText::FromString(Production.ProductionName),
+			FSlateIcon(),
+			FUIAction(FExecuteAction::CreateLambda([this, Production]() 
+				{
+					CustomizedCineAssembly->Modify();
+					CustomizedCineAssembly->Production = Production.ProductionID;
+				})),
+			NAME_None,
+			EUserInterfaceActionType::None
+		);
+	}
+
+	return MenuBuilder.MakeWidget();
+}
+
 #undef LOCTEXT_NAMESPACE
diff --git a/Engine/Plugins/VirtualProduction/CinematicAssemblyTools/Source/ShotsEditor/Private/CineAssemblyCustomization.h b/Engine/Plugins/VirtualProduction/CinematicAssemblyTools/Source/ShotsEditor/Private/CineAssemblyCustomization.h
--- a/Engine/Plugins/VirtualProduction/CinematicAssemblyTools/Source/ShotsEditor/Private/CineAssemblyCustomization.h
+++ b/Engine/Plugins/VirtualProduction/CinematicAssemblyTools/Source/ShotsEditor/Private/CineAssemblyCustomization.h
@@ -23,6 +23,11 @@ public:
 	//~ End IDetailCustomization interface */
 
 private:
+	/** Utility functions to customize each of the categories */
+	void CustomizeDefaultCategory(IDetailLayoutBuilder& DetailBuilder);
+	void CustomizeMetadataCategory(IDetailLayoutBuilder& DetailBuilder);
+	void CustomizeSubsequenceCategory(IDetailLayoutBuilder& DetailBuilder);
+
 	/**
 	 * Evaluates all of the customization's template strings with the naming tokens subsystem.
 	 * This function is throttled to only run at a set frequency, to avoid the potential to constantly query the naming tokens subsystem.
@@ -36,10 +41,10 @@ private:
 	void EvaluateTokenString(FTemplateString& TokenString);
 
 	/** Checks if the customized CineAssembly contains a SubAssembly name in its list of assets to create */
-	ECheckBoxState IsSubAssemblyChecked(UCineAssembly* Assembly, int32 Index) const;
+	ECheckBoxState IsSubAssemblyChecked(int32 Index) const;
 
 	/** Adds/Removes a SubAssembly name from the custmoized CineAssembly's list of SubAssemblies */
-	void SubAssemblyCheckStateChanged(ECheckBoxState CheckBoxState, UCineAssembly* Assembly, int32 Index);
+	void SubAssemblyCheckStateChanged(ECheckBoxState CheckBoxState, int32 Index);
 
 	/** Returns the template text for a SubAssembly template name */
 	FText GetTemplateText(int32 Index) const;
@@ -48,9 +53,18 @@ private:
 	FText GetResolvedText(int32 Index);
 
 	/** Modifies the template text of the SubAssembly template name and re-evaluates the token string to update the resolved text */
-	void OnTemplateTextCommitted(const FText& InText, ETextCommit::Type InCommitType, UCineAssembly* Assembly, int32 Index);
+	void OnTemplateTextCommitted(const FText& InText, ETextCommit::Type InCommitType, int32 Index);
+
+	/** Builds the drop-down menu list of productions */
+	TSharedRef<SWidget> BuildProductionNameMenu();
+
+	/** Determines whether the input asset should be filtered out of an object picker widget, based on whether it is of the input schema type */
+	bool ShouldFilterAssetBySchema(const FAssetData& InAssetData, FSoftObjectPath Schema);
 
 private:
+	/** The assembly being customized */
+	UCineAssembly* CustomizedCineAssembly;
+
 	/** Array of template names from the customized CineAssembly's schema */
 	TArray<FTemplateString> SubAssemblyNames;
 
diff --git a/Engine/Plugins/VirtualProduction/CinematicAssemblyTools/Source/ShotsEditor/Private/CineAssemblySchema.h b/Engine/Plugins/VirtualProduction/CinematicAssemblyTools/Source/ShotsEditor/Private/CineAssemblySchema.h
--- a/Engine/Plugins/VirtualProduction/CinematicAssemblyTools/Source/ShotsEditor/Private/CineAssemblySchema.h
+++ b/Engine/Plugins/VirtualProduction/CinematicAssemblyTools/Source/ShotsEditor/Private/CineAssemblySchema.h
@@ -94,6 +94,10 @@ public:
 	UPROPERTY(EditAnywhere, Category = "Default")
 	FString DefaultAssemblyPath;
 
+	/** Restricts assemblies made from this schema to using this Schema when picking a Parent Assembly */
+	UPROPERTY(EditAnywhere, Category = "Default", meta = (AllowedClasses = "/Script/ShotsEditor.CineAssemblySchema"))
+	FSoftObjectPath ParentSchema;
+
 	/** The thumbnail image to use for this schema and assemblies built from this schema */
 	UPROPERTY(EditAnywhere, Category = "Default")
 	TObjectPtr<UTexture2D> ThumbnailImage;
