ShotManagement: Fix assembly creation window to not require the use of a modal window.
The Create Assembly window is no longer a modal window. It now creates the transient Cine Assembly to be configured, and upon clicking the Create Asset button, it will trigger the creation of a new package for the assembly object.
The Edit Assembly window is now a different widget from the config window. It will spawn in a docktable nomad tab for the user to edit its properties.
Updated the BP function CreateAssembly() to use the same code path to create an asset as the creation window.
Added a static evaluation function to the Cine Assembly Naming Tokens.

#rnx

diff --git a/Engine/Plugins/VirtualProduction/CinematicAssemblyTools/Source/ShotsEditor/Private/CineAssembly.cpp b/Engine/Plugins/VirtualProduction/CinematicAssemblyTools/Source/ShotsEditor/Private/CineAssembly.cpp
--- a/Engine/Plugins/VirtualProduction/CinematicAssemblyTools/Source/ShotsEditor/Private/CineAssembly.cpp
+++ b/Engine/Plugins/VirtualProduction/CinematicAssemblyTools/Source/ShotsEditor/Private/CineAssembly.cpp
@@ -8,6 +8,7 @@
 #include "CineAssemblyNamingTokens.h"
 #include "HAL/FileManager.h"
 #include "MovieScene.h"
+#include "MovieSceneToolsProjectSettings.h"
 #include "NamingTokensEngineSubsystem.h"
 #include "Sections/MovieSceneSubSection.h"
 #include "Serialization/JsonReader.h"
@@ -25,6 +26,18 @@ UCineAssembly::UCineAssembly()
 	MetadataJsonObject = MakeShared<FJsonObject>();
 }
 
+void UCineAssembly::Initialize()
+{
+	Super::Initialize();
+
+	const UMovieSceneToolsProjectSettings* MovieSceneToolsProjectSettings = GetDefault<UMovieSceneToolsProjectSettings>();
+
+	const FFrameRate TickResolution = MovieScene->GetTickResolution();
+	const FFrameNumber DefaultStartFrame = (MovieSceneToolsProjectSettings->DefaultStartTime * TickResolution).FloorToFrame();
+	const int32 DefaultDuration = (MovieSceneToolsProjectSettings->DefaultDuration * TickResolution).FloorToFrame().Value;
+	MovieScene->SetPlaybackRange(DefaultStartFrame, DefaultDuration);
+}
+
 const UCineAssemblySchema* UCineAssembly::GetSchema() const
 {
 	return BaseSchema.Get();
@@ -52,6 +65,16 @@ void UCineAssembly::ChangeSchema(UCineAssemblySchema* InSchema)
 
 	BaseSchema = InSchema;
 
+	// Reset the assembly's name based on the schema template
+	if (BaseSchema)
+	{
+		AssemblyName.Template = BaseSchema->DefaultAssemblyName;
+	}
+	else
+	{
+		AssemblyName.Template = TEXT("");
+	}
+
 	// Add all metadata associated with the new schema (initialized to the default values for each field)
 	if (BaseSchema)
 	{
@@ -107,28 +130,19 @@ void UCineAssembly::CreateSubAssemblies()
 		return;
 	}
 
-	UNamingTokensEngineSubsystem* NamingTokensSubsystem = GEngine->GetEngineSubsystem<UNamingTokensEngineSubsystem>();
-
-	UCineAssemblyNamingTokensContext* NamingTokenContext = NewObject<UCineAssemblyNamingTokensContext>();
-	NamingTokenContext->Assembly = this;
-
-	FNamingTokenFilterArgs FilterArgs;
-	FilterArgs.AdditionalNamespacesToInclude.Add(UCineAssemblyNamingTokens::TokenNamespace);
-
 	// Get the path where the top-level assembly will be created so we can create other assets relative to it
 	const FAssetData AssemblyAssetData = FAssetData(this);
 	FString PackagePath = AssemblyAssetData.PackagePath.ToString();
 
 	// Remove the default assembly path from the top-level assembly's package path to get back to the "root" path for the schema's folder hierarchy
-	FNamingTokenResultData DefaultAssemblyPathResult = NamingTokensSubsystem->EvaluateTokenString(BaseSchema->DefaultAssemblyPath, FilterArgs, { NamingTokenContext });
-	PackagePath = PackagePath.Replace(*DefaultAssemblyPathResult.EvaluatedText.ToString(), TEXT(""));
+	FString ResolvedDefaultAssemblyPath = UCineAssemblyNamingTokens::GetResolvedText(BaseSchema->DefaultAssemblyPath, this).ToString();
+	PackagePath = PackagePath.Replace(*ResolvedDefaultAssemblyPath, TEXT(""));
 
 	// Create the default folders for this assembly, based on the schema
 	for (FTemplateString& FolderPath : DefaultFolderNames)
 	{
 		// Resolve any tokens found in the name of the SubAssembly before attempting to create it 
-		FNamingTokenResultData Result = NamingTokensSubsystem->EvaluateTokenString(FolderPath.Template, FilterArgs, { NamingTokenContext });
-		FolderPath.Resolved = Result.EvaluatedText;
+		FolderPath.Resolved = UCineAssemblyNamingTokens::GetResolvedText(FolderPath.Template, this);
 
 		if (FolderPath.Resolved.IsEmpty())
 		{
@@ -157,8 +171,7 @@ void UCineAssembly::CreateSubAssemblies()
 	for (FTemplateString& SubAssemblyName : SubAssemblyNames)
 	{
 		// Resolve any tokens found in the name of the SubAssembly before attempting to create it 
-		FNamingTokenResultData Result = NamingTokensSubsystem->EvaluateTokenString(SubAssemblyName.Template, FilterArgs, { NamingTokenContext });
-		SubAssemblyName.Resolved = Result.EvaluatedText;
+		SubAssemblyName.Resolved = UCineAssemblyNamingTokens::GetResolvedText(SubAssemblyName.Template, this);
 
 		const FString SubAssemblyFilename = FPaths::GetBaseFilename(SubAssemblyName.Resolved.ToString());
 
@@ -182,6 +195,8 @@ void UCineAssembly::CreateSubAssemblies()
 		if (SubAssembly)
 		{
 			SubAssembly->Level = Level;
+			SubAssembly->ParentAssembly = this;
+			SubAssembly->Production = Production;
 
 			const TRange<FFrameNumber> PlaybackRange = SubAssembly->GetMovieScene()->GetPlaybackRange();
 			const FFrameNumber StartFrame = PlaybackRange.GetLowerBoundValue();
diff --git a/Engine/Plugins/VirtualProduction/CinematicAssemblyTools/Source/ShotsEditor/Private/CineAssembly.h b/Engine/Plugins/VirtualProduction/CinematicAssemblyTools/Source/ShotsEditor/Private/CineAssembly.h
--- a/Engine/Plugins/VirtualProduction/CinematicAssemblyTools/Source/ShotsEditor/Private/CineAssembly.h
+++ b/Engine/Plugins/VirtualProduction/CinematicAssemblyTools/Source/ShotsEditor/Private/CineAssembly.h
@@ -25,6 +25,10 @@ public:
 
 	friend class SCineAssemblyConfigWindow;
 
+	//~ Begin ULevelSquence Interface
+	virtual void Initialize() override;
+	//~ End ULevelSquence Interface
+
 	//~ Begin UObject Interface
 	virtual void GetAssetRegistryTags(FAssetRegistryTagsContext Context) const override;
 	virtual void GetAssetRegistryTagMetadata(TMap<FName, FAssetRegistryTagMetadata>& OutMetadata) const override;
diff --git a/Engine/Plugins/VirtualProduction/CinematicAssemblyTools/Source/ShotsEditor/Private/CineAssemblyCustomization.cpp b/Engine/Plugins/VirtualProduction/CinematicAssemblyTools/Source/ShotsEditor/Private/CineAssemblyCustomization.cpp
--- a/Engine/Plugins/VirtualProduction/CinematicAssemblyTools/Source/ShotsEditor/Private/CineAssemblyCustomization.cpp
+++ b/Engine/Plugins/VirtualProduction/CinematicAssemblyTools/Source/ShotsEditor/Private/CineAssemblyCustomization.cpp
@@ -10,7 +10,6 @@
 #include "DetailLayoutBuilder.h"
 #include "DetailWidgetRow.h"
 #include "Framework/MultiBox/MultiBoxBuilder.h"
-#include "NamingTokensEngineSubsystem.h"
 #include "ProductionSettings.h"
 #include "PropertyCustomizationHelpers.h"
 #include "STemplateStringEditableTextBox.h"
@@ -395,9 +394,7 @@ void FCineAssemblyCustomization::EvaluateTokenStrings()
 
 void FCineAssemblyCustomization::EvaluateTokenString(FTemplateString& TokenString)
 {
-	UNamingTokensEngineSubsystem* NamingTokensSubsystem = GEngine->GetEngineSubsystem<UNamingTokensEngineSubsystem>();
-	FNamingTokenResultData Result = NamingTokensSubsystem->EvaluateTokenString(TokenString.Template, FNamingTokenFilterArgs({ UCineAssemblyNamingTokens::TokenNamespace }));
-	TokenString.Resolved = Result.EvaluatedText;
+	TokenString.Resolved = UCineAssemblyNamingTokens::GetResolvedText(TokenString.Template, CustomizedCineAssembly);
 }
 
 TSharedRef<SWidget> FCineAssemblyCustomization::BuildProductionNameMenu()
diff --git a/Engine/Plugins/VirtualProduction/CinematicAssemblyTools/Source/ShotsEditor/Private/CineAssemblyFactory.cpp b/Engine/Plugins/VirtualProduction/CinematicAssemblyTools/Source/ShotsEditor/Private/CineAssemblyFactory.cpp
--- a/Engine/Plugins/VirtualProduction/CinematicAssemblyTools/Source/ShotsEditor/Private/CineAssemblyFactory.cpp
+++ b/Engine/Plugins/VirtualProduction/CinematicAssemblyTools/Source/ShotsEditor/Private/CineAssemblyFactory.cpp
@@ -6,15 +6,9 @@
 #include "CineAssemblyNamingTokens.h"
 #include "ContentBrowserModule.h"
 #include "Framework/Application/SlateApplication.h"
-#include "Interfaces/IMainFrameModule.h"
 #include "IContentBrowserSingleton.h"
 #include "Modules/ModuleManager.h"
-#include "MovieScene.h"
-#include "MovieSceneToolsProjectSettings.h"
-#include "NamingTokensEngineSubsystem.h"
-#include "Tracks/MovieSceneSubTrack.h"
 #include "UI/CineAssembly/SCineAssemblyConfigWindow.h"
-#include "UObject/Package.h"
 
 #define LOCTEXT_NAMESPACE "CineAssemblyFactory"
 
@@ -23,8 +17,6 @@ UCineAssemblyFactory::UCineAssemblyFactory()
 	SupportedClass = UCineAssembly::StaticClass();
 	bCreateNew = true;
 	bEditAfterNew = true;
-
-	CineAssemblyToConfigure = nullptr;
 }
 
 bool UCineAssemblyFactory::CanCreateNew() const
@@ -32,108 +24,98 @@ bool UCineAssemblyFactory::CanCreateNew() const
 	return true;
 }
 
-UObject* UCineAssemblyFactory::FactoryCreateNew(UClass* Class, UObject* InParent, FName Name, EObjectFlags Flags, UObject* Context, FFeedbackContext* Warn)
+bool UCineAssemblyFactory::ConfigureProperties()
 {
-	IAssetTools& AssetTools = FModuleManager::GetModuleChecked<FAssetToolsModule>("AssetTools").Get();
-
-	if (CineAssemblyToConfigure)
-	{
-		if (const UCineAssemblySchema* Schema = CineAssemblyToConfigure->GetSchema())
-		{
-			if (!Schema->DefaultAssemblyPath.IsEmpty())
-			{
-				const FString ResolvedString = ResolveTokenString(Schema->DefaultAssemblyPath);
-
-				// Change the outer package to include the DefaultAssemblyPath
-				const FString PackageName = FPaths::GetPath(InParent->GetPathName()) / ResolvedString / Name.ToString();
-				InParent = CreatePackage(*PackageName);
-			}
-		}
-
-		// Update the name, outer package, and flags of the configuration asset to no longer be transient
-		CineAssemblyToConfigure->Rename(*Name.ToString(), InParent);
-		CineAssemblyToConfigure->SetFlags(Flags);
-		CineAssemblyToConfigure->ClearFlags(RF_Transient);
-
-		// Unregister the temporary default name that we registered earlier for this schema instance
-		AssetTools.UnregisterDefaultAssetNameForClass(UCineAssembly::StaticClass());
-	}
-	else
-	{
-		// Procedural assembly creation that does not use the configuration window will hit this path
-		CineAssemblyToConfigure = NewObject<UCineAssembly>(InParent, Name, Flags);
-	}
-
-	// Do the same setup that ULevelSequence assets do when they are created
-	CineAssemblyToConfigure->Initialize();
-
-	const UMovieSceneToolsProjectSettings* MovieSceneToolsProjectSettings = GetDefault<UMovieSceneToolsProjectSettings>();
+	IContentBrowserSingleton& ContentBrowser = FModuleManager::LoadModuleChecked<FContentBrowserModule>("ContentBrowser").Get();
+	const FString CurrentPath = ContentBrowser.GetCurrentPath().GetInternalPathString();
 
-	FFrameRate TickResolution = CineAssemblyToConfigure->GetMovieScene()->GetTickResolution();
-	CineAssemblyToConfigure->GetMovieScene()->SetPlaybackRange((MovieSceneToolsProjectSettings->DefaultStartTime * TickResolution).FloorToFrame(), (MovieSceneToolsProjectSettings->DefaultDuration * TickResolution).FloorToFrame().Value);
+	TSharedRef<SCineAssemblyConfigWindow> CineAssemblyConfigWindow = SNew(SCineAssemblyConfigWindow, CurrentPath);
+	FSlateApplication::Get().AddWindow(CineAssemblyConfigWindow);
 
-	CineAssemblyToConfigure->CreateSubAssemblies();
+	return false;
+}
 
-	// Refresh the content browser to make sure the original asset tile is not visible (since it may have been moved to the schema's DefaultAssemblyPath)
-	IContentBrowserSingleton& ContentBrowser = FModuleManager::LoadModuleChecked<FContentBrowserModule>("ContentBrowser").Get();
-	ContentBrowser.SetSelectedPaths({ ContentBrowser.GetCurrentPath().GetInternalPathString() }, true);
+UObject* UCineAssemblyFactory::FactoryCreateNew(UClass* Class, UObject* InParent, FName Name, EObjectFlags Flags, UObject* Context, FFeedbackContext* Warn)
+{
+	// Procedural assembly creation that does not use the configuration window will hit this path
+	UCineAssembly* NewAssembly = NewObject<UCineAssembly>(InParent, Name, Flags);
+	NewAssembly->Initialize();
 
-	return CineAssemblyToConfigure;
+	return NewAssembly;
 }
 
-bool UCineAssemblyFactory::ConfigureProperties()
+void UCineAssemblyFactory::MakeUniqueNameAndPath(UCineAssembly* ConfiguredAssembly, const FString& CreateAssetPath, FString& UniquePackageName, FString& UniqueAssetName)
 {
-	// Spawn a new window to configure the properties of the asset being created
-	IMainFrameModule& MainFrame = FModuleManager::LoadModuleChecked<IMainFrameModule>("MainFrame");
-	TSharedPtr<SWindow>	ParentWindow = MainFrame.GetParentWindow();
+	// Resolve the default assembly path
+	FString DefaultAssemblyPath;
+	if (const UCineAssemblySchema* Schema = ConfiguredAssembly->GetSchema())
+	{
+		if (!Schema->DefaultAssemblyPath.IsEmpty())
+		{
+			DefaultAssemblyPath = UCineAssemblyNamingTokens::GetResolvedText(Schema->DefaultAssemblyPath, ConfiguredAssembly).ToString();
+		}
+	}
 
-	// Create a new transient CineAssembly to configure in the UI.
-	// If the configuration is successful, this will turn into the persistent object created by this factory.
-	CineAssemblyToConfigure = NewObject<UCineAssembly>(GetTransientPackage(), NAME_None, RF_Transient);
+	// Ensure that the resolved assembly name is actually unique
+	const FString AssemblyName = ConfiguredAssembly->AssemblyName.Resolved.ToString();
+	const FString DesiredPackageName = CreateAssetPath / DefaultAssemblyPath / AssemblyName;
+	const FString DesiredSuffix = TEXT("");
 
-	TSharedRef<SCineAssemblyConfigWindow> CineAssemblyConfigWindow = SNew(SCineAssemblyConfigWindow, CineAssemblyToConfigure, ECineAssemblyConfigMode::CreateNew);
-	FSlateApplication::Get().AddModalWindow(CineAssemblyConfigWindow, ParentWindow);
+	IAssetTools& AssetTools = FModuleManager::GetModuleChecked<FAssetToolsModule>("AssetTools").Get();
+	AssetTools.CreateUniqueAssetName(DesiredPackageName, DesiredSuffix, UniquePackageName, UniqueAssetName);
 
-	if (CineAssemblyConfigWindow->ShouldCreateNewAsset() == false)
+	// If the assembly name was not unique, update the assembly template string
+	if (UniqueAssetName != AssemblyName)
 	{
-		return false;
+		ConfiguredAssembly->AssemblyName.Template = UniqueAssetName;
+		ConfiguredAssembly->AssemblyName.Resolved = FText::FromString(UniqueAssetName);
+
+		// It is possible that the default assembly path was dependent on the assembly name (through tokens).
+		// By updating the assembly name, we may also need to reevaluate the assembly path.
+		// However, once we change the assembly path, we need to verify uniqueness again (in the new path).
+		// Therefore, we recursively call this function until we end up with a combination of a unique assembly path and name.
+		MakeUniqueNameAndPath(ConfiguredAssembly, CreateAssetPath, UniquePackageName, UniqueAssetName);
 	}
+}
 
-	// One of the properties configured by the user in the config window is the assembly name.
-	// We want the actual asset name to match the value of this property, and we don't want the user to have to retype it.
-	// This workaround allows us to temporarily set a default name for the assembly asset. 
-	// After FactoryCreateNew() is finished, we will unregister this default name override. 
-	if (!CineAssemblyToConfigure->AssemblyName.Template.IsEmpty())
+void UCineAssemblyFactory::CreateConfiguredAssembly(UCineAssembly* ConfiguredAssembly, const FString& CreateAssetPath)
+{
+	// The intended use of this function is to take a pre-configured, transient assembly, create a valid package for it, and initialize it.
+	if (ConfiguredAssembly->GetPackage() != GetTransientPackage())
 	{
-		const FString ResolvedString = ResolveTokenString(CineAssemblyToConfigure->AssemblyName.Template);
-		CineAssemblyToConfigure->AssemblyName.Resolved = FText::FromString(ResolvedString);
-
-		IAssetTools& AssetTools = FModuleManager::GetModuleChecked<FAssetToolsModule>("AssetTools").Get();
-		AssetTools.RegisterDefaultAssetNameForClass(UCineAssembly::StaticClass(), ResolvedString);
+		return;
 	}
 
-	return true;
-}
+	// Evaluate the name of the assembly
+	ConfiguredAssembly->AssemblyName.Resolved = UCineAssemblyNamingTokens::GetResolvedText(ConfiguredAssembly->AssemblyName.Template, ConfiguredAssembly);
 
-FString UCineAssemblyFactory::ResolveTokenString(const FString& TemplateString)
-{
-	FString ResolvedString;
-	if (CineAssemblyToConfigure)
+	// If the assembly name is empty, assign it a valid default name
+	if (ConfiguredAssembly->AssemblyName.Resolved.IsEmpty())
 	{
-		// Add the cine assembly being created to the naming token context
-		UCineAssemblyNamingTokensContext* NamingTokenContext = NewObject<UCineAssemblyNamingTokensContext>();
-		NamingTokenContext->Assembly = CineAssemblyToConfigure;
+		ConfiguredAssembly->AssemblyName.Resolved = LOCTEXT("NewCineAssemblyName", "NewCineAssembly");
+		ConfiguredAssembly->AssemblyName.Template = ConfiguredAssembly->AssemblyName.Resolved.ToString();
+	}
 
-		FNamingTokenFilterArgs FilterArgs;
-		FilterArgs.AdditionalNamespacesToInclude.Add(UCineAssemblyNamingTokens::TokenNamespace);
+	FString UniquePackageName;
+	FString UniqueAssetName;
+	MakeUniqueNameAndPath(ConfiguredAssembly, CreateAssetPath, UniquePackageName, UniqueAssetName);
 
-		UNamingTokensEngineSubsystem* NamingTokensSubsystem = GEngine->GetEngineSubsystem<UNamingTokensEngineSubsystem>();
-		FNamingTokenResultData Result = NamingTokensSubsystem->EvaluateTokenString(TemplateString, FilterArgs, { NamingTokenContext });
+	// The input assembly object was created in the transient package while its properties were configured.
+	// Now, we can create a real package for it, rename it, and update its object flags.
+	UPackage* Package = CreatePackage(*UniquePackageName);
+	ConfiguredAssembly->Rename(*UniqueAssetName, Package);
 
-		ResolvedString = Result.EvaluatedText.ToString();
-	}
+	const EObjectFlags Flags = RF_Public | RF_Standalone | RF_Transactional;
+	ConfiguredAssembly->SetFlags(Flags);
+	ConfiguredAssembly->ClearFlags(RF_Transient);
+
+	// Do the same setup that ULevelSequence assets do when they are created
+	ConfiguredAssembly->Initialize();
+	ConfiguredAssembly->CreateSubAssemblies();
 
-	return ResolvedString;
+	// Refresh the content browser to make any new assets and folders are immediately visible
+	IContentBrowserSingleton& ContentBrowser = FModuleManager::LoadModuleChecked<FContentBrowserModule>("ContentBrowser").Get();
+	ContentBrowser.SetSelectedPaths({ CreateAssetPath }, true);
 }
 
 #undef LOCTEXT_NAMESPACE
diff --git a/Engine/Plugins/VirtualProduction/CinematicAssemblyTools/Source/ShotsEditor/Private/CineAssemblyFactory.h b/Engine/Plugins/VirtualProduction/CinematicAssemblyTools/Source/ShotsEditor/Private/CineAssemblyFactory.h
--- a/Engine/Plugins/VirtualProduction/CinematicAssemblyTools/Source/ShotsEditor/Private/CineAssemblyFactory.h
+++ b/Engine/Plugins/VirtualProduction/CinematicAssemblyTools/Source/ShotsEditor/Private/CineAssemblyFactory.h
@@ -16,6 +16,10 @@ class UCineAssemblyFactory : public UFactory
 {
 	GENERATED_BODY()
 
+public:
+	/** Takes a pre-configured, transient assembly, creates a valid package for it, and initializes it */
+	static void CreateConfiguredAssembly(UCineAssembly* ConfiguredAssembly, const FString& CreateAssetPath);
+
 protected:
 	UCineAssemblyFactory();
 
@@ -25,11 +29,6 @@ protected:
 	virtual bool ConfigureProperties() override;
 	// End UFactory Interface
 
-private:
-	/** Utility function that evaluates the input string with the naming tokens subsystem and returns the resolved string */
-	FString ResolveTokenString(const FString& TemplateString);
-
-private:
-	/** Transient CineAssembly that will be configured by the user and used as a template for the actual asset created by this factory */
-	UCineAssembly* CineAssemblyToConfigure;
+	/** Recursively evaluates the default assembly name and path until a unique combination is found */
+	static void MakeUniqueNameAndPath(UCineAssembly* ConfiguredAssembly, const FString& CreateAssetPath, FString& UniquePackageName, FString& UniqueAssetName);
 };
diff --git a/Engine/Plugins/VirtualProduction/CinematicAssemblyTools/Source/ShotsEditor/Private/CineAssemblyNamingTokens.cpp b/Engine/Plugins/VirtualProduction/CinematicAssemblyTools/Source/ShotsEditor/Private/CineAssemblyNamingTokens.cpp
--- a/Engine/Plugins/VirtualProduction/CinematicAssemblyTools/Source/ShotsEditor/Private/CineAssemblyNamingTokens.cpp
+++ b/Engine/Plugins/VirtualProduction/CinematicAssemblyTools/Source/ShotsEditor/Private/CineAssemblyNamingTokens.cpp
@@ -2,12 +2,28 @@
 
 #include "CineAssemblyNamingTokens.h"
 
+#include "NamingTokensEngineSubsystem.h"
+
 #include "Algo/IndexOf.h"
 
 #define LOCTEXT_NAMESPACE "CineAssemblyNamingTokens"
 
 FString UCineAssemblyNamingTokens::TokenNamespace = TEXT("cat");
 
+FText UCineAssemblyNamingTokens::GetResolvedText(const FString& InStringToEvaluate, UCineAssembly* InAssembly)
+{
+	UCineAssemblyNamingTokensContext* NamingTokenContext = NewObject<UCineAssemblyNamingTokensContext>();
+	NamingTokenContext->Assembly = InAssembly;
+
+	FNamingTokenFilterArgs FilterArgs;
+	FilterArgs.AdditionalNamespacesToInclude.Add(TokenNamespace);
+
+	UNamingTokensEngineSubsystem* NamingTokensSubsystem = GEngine->GetEngineSubsystem<UNamingTokensEngineSubsystem>();
+	FNamingTokenResultData Result = NamingTokensSubsystem->EvaluateTokenString(InStringToEvaluate, FilterArgs, { NamingTokenContext });
+
+	return Result.EvaluatedText;
+}
+
 UCineAssemblyNamingTokens::UCineAssemblyNamingTokens()
 {
 	Namespace = TokenNamespace;
@@ -82,11 +98,6 @@ void UCineAssemblyNamingTokens::OnPostEvaluate_Implementation()
 	Context = nullptr;
 }
 
-void UCineAssemblyNamingTokens::SetDefaultCineAssembly(UCineAssembly* InAssembly)
-{
-	DefaultAssembly = InAssembly;
-}
-
 void UCineAssemblyNamingTokens::AddMetadataToken(const FString& InTokenKey)
 {
 	// If the token already exists, remove it, and replace it with the new one
@@ -120,15 +131,11 @@ void UCineAssemblyNamingTokens::AddMetadataToken(const FString& InTokenKey)
 
 FText UCineAssemblyNamingTokens::ExecuteTokenFunc(TFunction<FText(TWeakObjectPtr<UCineAssembly>)> TokenFunc)
 {
-	// Try evaluating the token function using the assembly from the context. If there was no context, try using the default assembly. 
+	// Try evaluating the token function using the assembly from the context.
 	if (Context)
 	{
 		return TokenFunc(Context->Assembly);
 	}
-	else if (DefaultAssembly.IsValid())
-	{
-		return TokenFunc(DefaultAssembly);
-	}
 
 	// Cannot evaluate token function without a valid assembly
 	return FText::GetEmpty();
diff --git a/Engine/Plugins/VirtualProduction/CinematicAssemblyTools/Source/ShotsEditor/Private/CineAssemblyNamingTokens.h b/Engine/Plugins/VirtualProduction/CinematicAssemblyTools/Source/ShotsEditor/Private/CineAssemblyNamingTokens.h
--- a/Engine/Plugins/VirtualProduction/CinematicAssemblyTools/Source/ShotsEditor/Private/CineAssemblyNamingTokens.h
+++ b/Engine/Plugins/VirtualProduction/CinematicAssemblyTools/Source/ShotsEditor/Private/CineAssemblyNamingTokens.h
@@ -14,7 +14,7 @@ class UCineAssemblyNamingTokensContext : public UObject
 	GENERATED_BODY()
 
 public:
-	/** The sssembly asset to use when evaluating tokens */
+	/** The assembly asset to use when evaluating tokens */
 	TWeakObjectPtr<UCineAssembly> Assembly;
 };
 
@@ -37,20 +37,17 @@ public:
 
 	static FString TokenNamespace;
 
+	/** Utility function that will evaluate the input string with the naming tokens subsystem and return the result evaluated text using the input assembly */
+	static FText GetResolvedText(const FString& InStringToEvaluate, UCineAssembly* InAssembly);
+
 	/** Adds a metadata token with the same name as the input key which will return the result of value of that metadata if found in the Cine Assembly */
 	void AddMetadataToken(const FString& InTokenKey);
 
-	/** Sets a default assembly asset to use (useful when there is only ever one assembly needed in a given context to skip passing a context) */
-	void SetDefaultCineAssembly(UCineAssembly* InAssembly);
-
 private:
 	/** Runs the input Token evaluation function, which takes a Cine Assembly as input and outputs FText */
 	FText ExecuteTokenFunc(TFunction<FText(TWeakObjectPtr<UCineAssembly>)> TokenFunc);
 
 private:
-	/** The assembly asset to use when evaluating tokens if no context is provided */
-	TWeakObjectPtr<UCineAssembly> DefaultAssembly;
-
 	/** The current context to use when evaluating tokens */
 	TObjectPtr<UCineAssemblyNamingTokensContext> Context;
 };
diff --git a/Engine/Plugins/VirtualProduction/CinematicAssemblyTools/Source/ShotsEditor/Private/ProductionFunctionLibrary.cpp b/Engine/Plugins/VirtualProduction/CinematicAssemblyTools/Source/ShotsEditor/Private/ProductionFunctionLibrary.cpp
--- a/Engine/Plugins/VirtualProduction/CinematicAssemblyTools/Source/ShotsEditor/Private/ProductionFunctionLibrary.cpp
+++ b/Engine/Plugins/VirtualProduction/CinematicAssemblyTools/Source/ShotsEditor/Private/ProductionFunctionLibrary.cpp
@@ -3,13 +3,10 @@
 #include "ProductionFunctionLibrary.h"
 
 #include "Algo/Find.h"
-#include "AssetToolsModule.h"
 #include "CineAssembly.h"
-#include "CineAssemblyNamingTokens.h"
+#include "CineAssemblyFactory.h"
 #include "CineAssemblySchema.h"
-#include "MovieScene.h"
-#include "MovieSceneToolsProjectSettings.h"
-#include "NamingTokensEngineSubsystem.h"
+#include "UObject/Package.h"
 
 const TArray<FCinematicProduction> UProductionFunctionLibrary::GetAllProductions()
 {
@@ -89,52 +86,19 @@ void UProductionFunctionLibrary::RenameProduction(FGuid ProductionID, FString Ne
 
 UCineAssembly* UProductionFunctionLibrary::CreateAssembly(UCineAssemblySchema* Schema, TSoftObjectPtr<UWorld> Level, TMap<FString, FString> Metadata, const FString& Path, const FString& Name, bool bUseDefaultNameFromSchema)
 {
-	IAssetTools& AssetTools = FModuleManager::GetModuleChecked<FAssetToolsModule>("AssetTools").Get();
-
-	// We need a valid name for the initial asset creation, even if we are going to rename it later.
-	FString InitialName = Name;
-	if (InitialName.IsEmpty())
-	{
-		InitialName = TEXT("NewCineAssembly");
-	}
-
-	FString UniquePackageName;
-	FString UniqueAssetName;
-	AssetTools.CreateUniqueAssetName(Path / InitialName, TEXT(""), UniquePackageName, UniqueAssetName);
-
-	// Create the new assembly with a unique name. 
-	// No factory is provided, because we need to set the schema and metadata on the Assembly asset first, before we can create any SubAssemblies.
-	UCineAssembly* NewAssembly = Cast<UCineAssembly>(AssetTools.CreateAsset(UniqueAssetName, FPaths::GetPath(UniquePackageName), UCineAssembly::StaticClass(), nullptr));
-
-	NewAssembly->Level = Level.ToSoftObjectPath();
+	UCineAssembly* NewAssembly = NewObject<UCineAssembly>(GetTransientPackage(), NAME_None, RF_Transient);
 	NewAssembly->SetSchema(Schema);
-
-	// Do the same setup that ULevelSequence assets do when they are created
-	NewAssembly->Initialize();
-
-	const UMovieSceneToolsProjectSettings* MovieSceneToolsProjectSettings = GetDefault<UMovieSceneToolsProjectSettings>();
-
-	FFrameRate TickResolution = NewAssembly->GetMovieScene()->GetTickResolution();
-	NewAssembly->GetMovieScene()->SetPlaybackRange((MovieSceneToolsProjectSettings->DefaultStartTime * TickResolution).FloorToFrame(), (MovieSceneToolsProjectSettings->DefaultDuration * TickResolution).FloorToFrame().Value);
-
+	NewAssembly->SetLevel(Level);
+	
 	AddMetadataToAssembly(NewAssembly, Metadata);
 
-	if (Schema && bUseDefaultNameFromSchema)
+	if (!Schema || !bUseDefaultNameFromSchema)
 	{
-		NewAssembly->AssemblyName.Template = Schema->DefaultAssemblyName;
-
-		if (!NewAssembly->AssemblyName.Template.IsEmpty())
-		{
-			RenameAssembly(NewAssembly);
-		}
-	}
-	else
-	{
-		NewAssembly->AssemblyName.Template = UniqueAssetName;
-		NewAssembly->AssemblyName.Resolved = FText::FromString(UniqueAssetName);
+		NewAssembly->AssemblyName.Template = Name;
+		NewAssembly->AssemblyName.Resolved = FText::FromString(Name);
 	}
 
-	NewAssembly->CreateSubAssemblies();
+	UCineAssemblyFactory::CreateConfiguredAssembly(NewAssembly, Path);
 
 	return NewAssembly;
 }
@@ -178,59 +142,3 @@ void UProductionFunctionLibrary::AddMetadataToAssembly(UCineAssembly* Assembly,
 		}
 	}
 }
-
-void UProductionFunctionLibrary::RenameAssembly(UCineAssembly* Assembly)
-{
-	IAssetTools& AssetTools = FModuleManager::GetModuleChecked<FAssetToolsModule>("AssetTools").Get();
-
-	// Add the cine assembly being created to the naming token context
-	UCineAssemblyNamingTokensContext* NamingTokenContext = NewObject<UCineAssemblyNamingTokensContext>();
-	NamingTokenContext->Assembly = Assembly;
-
-	FNamingTokenFilterArgs FilterArgs;
-	FilterArgs.AdditionalNamespacesToInclude.Add(UCineAssemblyNamingTokens::TokenNamespace);
-
-	// Evaluate the default assembly name to resolve any tokens 
-	UNamingTokensEngineSubsystem* NamingTokensSubsystem = GEngine->GetEngineSubsystem<UNamingTokensEngineSubsystem>();
-	FNamingTokenResultData Result = NamingTokensSubsystem->EvaluateTokenString(Assembly->AssemblyName.Template, FilterArgs, { NamingTokenContext });
-
-	Assembly->AssemblyName.Resolved = Result.EvaluatedText;
-
-	// Resolve the default assembly path from teh schema, which will be added to the new desired package path
-	FString DefaultAssemblyPath;
-	if (const UCineAssemblySchema* Schema = Assembly->GetSchema())
-	{
-		if (!Schema->DefaultAssemblyPath.IsEmpty())
-		{
-			Result = NamingTokensSubsystem->EvaluateTokenString(Schema->DefaultAssemblyPath, FilterArgs, { NamingTokenContext });
-			DefaultAssemblyPath = Result.EvaluatedText.ToString();
-		}
-	}
-
-	// We now need to rename the asset (which was originally created with a default name) to the resolved default assembly name
-	const FSoftObjectPath OldObjectPath = FSoftObjectPath(Assembly);
-
-	// Ensure that the resolved assembly name is actually unique
-	const FString DesiredPackageName = FPaths::GetPath(Assembly->GetPathName()) / DefaultAssemblyPath / Assembly->AssemblyName.Resolved.ToString();
-
-	FString UniquePackageName;
-	FString UniqueAssetName;
-	AssetTools.CreateUniqueAssetName(DesiredPackageName, TEXT(""), UniquePackageName, UniqueAssetName);
-
-	// If we had to add anything to the assembly name to make it unique, update the assembly name template to match
-	if (UniqueAssetName != Assembly->AssemblyName.Resolved.ToString())
-	{
-		Assembly->AssemblyName.Template = UniqueAssetName;
-		Assembly->AssemblyName.Resolved = FText::FromString(UniqueAssetName);
-	}
-
-	const FTopLevelAssetPath NewAssetPath(*UniquePackageName, *UniqueAssetName);
-	const FSoftObjectPath NewObjectPath(NewAssetPath);
-
-	if (OldObjectPath != NewObjectPath)
-	{
-		FAssetRenameData RenameData = FAssetRenameData(OldObjectPath, NewObjectPath);
-		AssetTools.RenameAssets({ RenameData });
-	}
-}
-
diff --git a/Engine/Plugins/VirtualProduction/CinematicAssemblyTools/Source/ShotsEditor/Private/ProductionFunctionLibrary.h b/Engine/Plugins/VirtualProduction/CinematicAssemblyTools/Source/ShotsEditor/Private/ProductionFunctionLibrary.h
--- a/Engine/Plugins/VirtualProduction/CinematicAssemblyTools/Source/ShotsEditor/Private/ProductionFunctionLibrary.h
+++ b/Engine/Plugins/VirtualProduction/CinematicAssemblyTools/Source/ShotsEditor/Private/ProductionFunctionLibrary.h
@@ -72,7 +72,6 @@ public:
 	static UCineAssembly* CreateAssembly(UCineAssemblySchema* Schema, TSoftObjectPtr<UWorld> Level, TMap<FString, FString> Metadata, const FString& Path, const FString& Name, bool bUseDefaultNameFromSchema = true);
 
 private:
-
+	/** Adds the input map of metadata to the input assembly */
 	static void AddMetadataToAssembly(UCineAssembly* Assembly, TMap<FString, FString> Metadata);
-	static void RenameAssembly(UCineAssembly* Assembly);
 };
diff --git a/Engine/Plugins/VirtualProduction/CinematicAssemblyTools/Source/ShotsEditor/Private/ShotsEditorModule.cpp b/Engine/Plugins/VirtualProduction/CinematicAssemblyTools/Source/ShotsEditor/Private/ShotsEditorModule.cpp
--- a/Engine/Plugins/VirtualProduction/CinematicAssemblyTools/Source/ShotsEditor/Private/ShotsEditorModule.cpp
+++ b/Engine/Plugins/VirtualProduction/CinematicAssemblyTools/Source/ShotsEditor/Private/ShotsEditorModule.cpp
@@ -48,6 +48,9 @@ private:
 	/** Called whenever a job has its sequence set. */
 	void OnSequenceSet(UMoviePipelineExecutorJob* Job, ULevelSequence* Sequence);
 
+	/** Opens a window to edit the properties of the input assembly */
+	void OpenAssemblyForEdit(UCineAssembly* Assembly);
+
 private:
 	FDelegateHandle OnSequenceSetHandle;
 };
@@ -123,24 +126,25 @@ void FShotsEditorModule::RegisterMenus()
 		const FName AssetActionSectionName = TEXT("GetAssetActions");
 		FToolMenuSection& Section = Menu->FindOrAddSection(AssetActionSectionName);
 
-		Section.AddDynamicEntry("CineAssemblyActions", FNewToolMenuSectionDelegate::CreateLambda([](FToolMenuSection& InSection)
+		Section.AddDynamicEntry("CineAssemblyActions", FNewToolMenuSectionDelegate::CreateLambda([this](FToolMenuSection& InSection)
 		{
 			if (UContentBrowserAssetContextMenuContext* Context = InSection.FindContext<UContentBrowserAssetContextMenuContext>())
 			{
 				if ((Context->SelectedAssets.Num() == 1) && Context->SelectedAssets[0].IsInstanceOf(UCineAssembly::StaticClass()))
 				{
-					const FAssetData CineAssemblyData = Context->SelectedAssets[0];
+					const FAssetData& CineAssemblyData = Context->SelectedAssets[0];
+					UCineAssembly* CineAssembly = Cast<UCineAssembly>(CineAssemblyData.GetAsset());
 
 					InSection.AddMenuEntry(
 						"OpenInCurrentMapEntry",
 						LOCTEXT("OpenInCurrentMap", "Open in Current Map"),
 						LOCTEXT("OpenInCurrentMapTooltip", "Opens the level sequence in Sequencer but does not automatically load the associated map"),
 						FSlateIcon(FAppStyle::GetAppStyleSetName(), "LevelEditor.OpenCinematic"),
-						FUIAction(FExecuteAction::CreateLambda([CineAssemblyData]()
+						FUIAction(FExecuteAction::CreateLambda([CineAssembly]()
 						{
 							// Note: Normally, when a cine assembly is opened in the editor, its associated map is loaded during the PrepareToActiveAssets step.
 							// By directly opening the asset editor (Sequencer) for this asset, we can skip that map loading step. 
-							GEditor->GetEditorSubsystem<UAssetEditorSubsystem>()->OpenEditorForAssets({ CineAssemblyData.GetAsset() });
+							GEditor->GetEditorSubsystem<UAssetEditorSubsystem>()->OpenEditorForAssets({ CineAssembly });
 						})),
 						EUserInterfaceActionType::None
 					);
@@ -150,15 +154,7 @@ void FShotsEditorModule::RegisterMenus()
 						LOCTEXT("OpenInAssetEditor", "Edit Properties"),
 						LOCTEXT("OpenInAssetEditorTooltip", "Edit Properties"),
 						FSlateIcon(FAppStyle::GetAppStyleSetName(), "Icons.Edit"),
-						FUIAction(FExecuteAction::CreateLambda([CineAssemblyData]()
-						{
-							// Spawn a new window to configure the properties of the asset being created
-							IMainFrameModule& MainFrame = FModuleManager::LoadModuleChecked<IMainFrameModule>("MainFrame");
-							TSharedPtr<SWindow>	ParentWindow = MainFrame.GetParentWindow();
-
-							TSharedRef<SCineAssemblyConfigWindow> CineAssemblyConfigWindow = SNew(SCineAssemblyConfigWindow, Cast<UCineAssembly>(CineAssemblyData.GetAsset()), ECineAssemblyConfigMode::Edit);
-							FSlateApplication::Get().AddModalWindow(CineAssemblyConfigWindow, ParentWindow);
-						})),
+						FUIAction(FExecuteAction::CreateRaw(this, &FShotsEditorModule::OpenAssemblyForEdit, CineAssembly)),
 						EUserInterfaceActionType::None
 					);
 				}
@@ -167,6 +163,22 @@ void FShotsEditorModule::RegisterMenus()
 	}
 }
 
+void FShotsEditorModule::OpenAssemblyForEdit(UCineAssembly* Assembly)
+{
+	TSharedRef<SCineAssemblyEditWidget> CineAssemblyWidget = SNew(SCineAssemblyEditWidget, Assembly);
+
+	const FVector2D DefaultWindowSize = FVector2D(600, 800);
+	TSharedRef<SWindow> EditWindow = SNew(SWindow)
+		.Title(LOCTEXT("EditCineAssemblyTitle", "Edit Cine Assembly"))
+		.SizingRule(ESizingRule::UserSized)
+		.ClientSize(DefaultWindowSize)
+		.SupportsMaximize(false)
+		.SupportsMinimize(false);
+
+	EditWindow->SetContent(CineAssemblyWidget);
+	FSlateApplication::Get().AddWindow(EditWindow);
+}
+
 void FShotsEditorModule::OnSequenceSet(UMoviePipelineExecutorJob* Job, ULevelSequence* Sequence)
 {
 	if (Job)
diff --git a/Engine/Plugins/VirtualProduction/CinematicAssemblyTools/Source/ShotsEditor/Private/UI/CineAssembly/SCineAssemblyConfigWindow.cpp b/Engine/Plugins/VirtualProduction/CinematicAssemblyTools/Source/ShotsEditor/Private/UI/CineAssembly/SCineAssemblyConfigWindow.cpp
--- a/Engine/Plugins/VirtualProduction/CinematicAssemblyTools/Source/ShotsEditor/Private/UI/CineAssembly/SCineAssemblyConfigWindow.cpp
+++ b/Engine/Plugins/VirtualProduction/CinematicAssemblyTools/Source/ShotsEditor/Private/UI/CineAssembly/SCineAssemblyConfigWindow.cpp
@@ -2,6 +2,7 @@
 
 #include "SCineAssemblyConfigWindow.h"
 
+#include "CineAssemblyFactory.h"
 #include "CineAssemblyNamingTokens.h"
 #include "ContentBrowserModule.h"
 #include "Editor.h"
@@ -9,8 +10,8 @@
 #include "Interfaces/IPluginManager.h"
 #include "IStructureDetailsView.h"
 #include "Modules/ModuleManager.h"
-#include "NamingTokensEngineSubsystem.h"
 #include "PropertyEditorModule.h"
+#include "UObject/Package.h"
 #include "ShotsEditorStyle.h"
 #include "STemplateStringEditableTextBox.h"
 #include "Widgets/SBoxPanel.h"
@@ -24,39 +25,37 @@
 
 #define LOCTEXT_NAMESPACE "SCineAssemblyConfigWindow"
 
-void SCineAssemblyConfigWindow::Construct(const FArguments& InArgs, UCineAssembly* InCineAssembly, ECineAssemblyConfigMode InConfigMode)
+namespace UE::CineAssemblyWidgets::Private
 {
-	CineAssemblyToConfigure = InCineAssembly;
-	Mode = InConfigMode;
-
-	// When creating a new CineAssembly, initialize its Level to the currently open level
-	if (Mode == ECineAssemblyConfigMode::CreateNew)
+	TSharedRef<IDetailsView> MakeDetailsWidget(UCineAssembly* Assembly)
 	{
-		UWorld* CurrentWorld = GEditor->GetEditorWorldContext().World();
-		CineAssemblyToConfigure->Level = FSoftObjectPath(CurrentWorld);
+		FPropertyEditorModule& PropertyEditor = FModuleManager::Get().LoadModuleChecked<FPropertyEditorModule>(TEXT("PropertyEditor"));
+
+		FDetailsViewArgs DetailsViewArgs;
+		DetailsViewArgs.bAllowSearch = false;
+		DetailsViewArgs.bHideSelectionTip = true;
+
+		TSharedRef<IDetailsView> DetailsView = PropertyEditor.CreateDetailView(DetailsViewArgs);
+		DetailsView->SetObject(Assembly, true);
+
+		return DetailsView;
 	}
+}
 
-	const FVector2D DefaultWindowSize = FVector2D(1400, 750);
+void SCineAssemblyConfigWindow::Construct(const FArguments& InArgs, const FString& InCreateAssetPath)
+{
+	CreateAssetPath = InCreateAssetPath;
 
-	GetOnWindowClosedEvent().AddSP(this, &SCineAssemblyConfigWindow::OnWindowClosed);
+	// Create a new transient CineAssembly to configure in the UI.
+	CineAssemblyToConfigure = NewObject<UCineAssembly>(GetTransientPackage(), NAME_None, RF_Transient);
 
-	// Set the assembly being configured as the default assembly to use when evaluating naming tokens to avoid passing a context object every time we evaluate
-	UNamingTokensEngineSubsystem* NamingTokensSubsystem = GEngine->GetEngineSubsystem<UNamingTokensEngineSubsystem>();
-	UCineAssemblyNamingTokens* CineAssemblyNamingTokens = Cast<UCineAssemblyNamingTokens>(NamingTokensSubsystem->GetNamingTokens(UCineAssemblyNamingTokens::TokenNamespace));
-	CineAssemblyNamingTokens->SetDefaultCineAssembly(CineAssemblyToConfigure);
+	UWorld* CurrentWorld = GEditor->GetEditorWorldContext().World();
+	CineAssemblyToConfigure->Level = FSoftObjectPath(CurrentWorld);
 
-	FText WindowTitle;
-	if (Mode == ECineAssemblyConfigMode::CreateNew)
-	{
-		WindowTitle = LOCTEXT("WindowTitleCreateNew", "Create Cine Assembly");
-	}
-	else if (Mode == ECineAssemblyConfigMode::Edit)
-	{
-		WindowTitle = LOCTEXT("WindowTitleEdit", "Edit Cine Assembly");
-	}
+	const FVector2D DefaultWindowSize = FVector2D(1400, 750);
 
 	SWindow::Construct(SWindow::FArguments()
-		.Title(WindowTitle)
+		.Title(LOCTEXT("WindowTitleCreateNew", "Create Cine Assembly"))
 		.SizingRule(ESizingRule::UserSized)
 		.ClientSize(DefaultWindowSize)
 		.SupportsMaximize(false)
@@ -103,14 +102,6 @@ void SCineAssemblyConfigWindow::Construct(const FArguments& InArgs, UCineAssembl
 		]);
 }
 
-void SCineAssemblyConfigWindow::OnWindowClosed(const TSharedRef<SWindow>& InWindow)
-{
-	// Set the assembly being configured as the default assembly to use when evaluating naming tokens to avoid passing a context object every time we evaluate
-	UNamingTokensEngineSubsystem* NamingTokensSubsystem = GEngine->GetEngineSubsystem<UNamingTokensEngineSubsystem>();
-	UCineAssemblyNamingTokens* CineAssemblyNamingTokens = Cast<UCineAssemblyNamingTokens>(NamingTokensSubsystem->GetNamingTokens(UCineAssemblyNamingTokens::TokenNamespace));
-	CineAssemblyNamingTokens->SetDefaultCineAssembly(nullptr);
-}
-
 TSharedRef<SWidget> SCineAssemblyConfigWindow::MakeCineTemplatePanel()
 {
 	FContentBrowserModule& ContentBrowserModule = FModuleManager::Get().LoadModuleChecked<FContentBrowserModule>("ContentBrowser");
@@ -170,12 +161,10 @@ void SCineAssemblyConfigWindow::OnSchemaSelected(const FAssetData& AssetData)
 	if (UCineAssemblySchema* CineAssemblySchema = Cast<UCineAssemblySchema>(AssetData.GetAsset()))
 	{
 		SelectedSchema = CineAssemblySchema;
-		CineAssemblyToConfigure->AssemblyName.Template = SelectedSchema->DefaultAssemblyName;
 	}
 	else
 	{
 		SelectedSchema = nullptr;
-		CineAssemblyToConfigure->AssemblyName.Template = TEXT("");
 	}
 
 	CineAssemblyToConfigure->ChangeSchema(SelectedSchema);
@@ -189,12 +178,14 @@ void SCineAssemblyConfigWindow::OnSchemaSelected(const FAssetData& AssetData)
 
 TSharedRef<SWidget> SCineAssemblyConfigWindow::MakeInfoPanel()
 {
+	DetailsView = UE::CineAssemblyWidgets::Private::MakeDetailsWidget(CineAssemblyToConfigure);
+
 	TabSwitcher = SNew(SWidgetSwitcher)
 		+ SWidgetSwitcher::Slot()
 			[ MakeOverviewWidget() ]
 
 		+ SWidgetSwitcher::Slot()
-			[ MakeDetailsWidget() ]
+			[ DetailsView.ToSharedRef() ]
 
 		+ SWidgetSwitcher::Slot()
 			[ MakeHierarchyWidget() ]
@@ -284,10 +275,7 @@ void SCineAssemblyConfigWindow::EvaluateTokenString(FTemplateString& StringToEva
 	FDateTime CurrentTime = FDateTime::Now();
 	if ((CurrentTime - LastTokenUpdateTime).GetSeconds() >= 1.0f)
 	{
-		UNamingTokensEngineSubsystem* NamingTokensSubsystem = GEngine->GetEngineSubsystem<UNamingTokensEngineSubsystem>();
-		FNamingTokenResultData Result = NamingTokensSubsystem->EvaluateTokenString(StringToEvaluate.Template, FNamingTokenFilterArgs({ UCineAssemblyNamingTokens::TokenNamespace }));
-
-		StringToEvaluate.Resolved = Result.EvaluatedText;
+		StringToEvaluate.Resolved = UCineAssemblyNamingTokens::GetResolvedText(StringToEvaluate.Template, CineAssemblyToConfigure);
 		LastTokenUpdateTime = CurrentTime;
 	}
 }
@@ -362,20 +350,6 @@ TSharedRef<SWidget> SCineAssemblyConfigWindow::MakeOverviewWidget()
 		];
 }
 
-TSharedRef<SWidget> SCineAssemblyConfigWindow::MakeDetailsWidget()
-{
-	FPropertyEditorModule& PropertyEditor = FModuleManager::Get().LoadModuleChecked<FPropertyEditorModule>(TEXT("PropertyEditor"));
-
-	FDetailsViewArgs DetailsViewArgs;
-	DetailsViewArgs.bAllowSearch = false;
-	DetailsViewArgs.bHideSelectionTip = true;
-
-	DetailsView = PropertyEditor.CreateDetailView(DetailsViewArgs);
-	DetailsView->SetObject(CineAssemblyToConfigure, true);
-
-	return DetailsView.ToSharedRef();
-}
-
 TSharedRef<SWidget> SCineAssemblyConfigWindow::MakeHierarchyWidget()
 {
 	HierarchyTreeView = SNew(STreeView<TSharedPtr<FHierarchyTreeItem>>)
@@ -484,11 +458,8 @@ void SCineAssemblyConfigWindow::PopulateHierarchyTree()
 
 void SCineAssemblyConfigWindow::EvaluateHierarchyTokensRecursive(TSharedPtr<FHierarchyTreeItem> TreeItem)
 {
-	UNamingTokensEngineSubsystem* NamingTokensSubsystem = GEngine->GetEngineSubsystem<UNamingTokensEngineSubsystem>();
-
-	// Evalaute the token template string for this tree item
-	FNamingTokenResultData Result = NamingTokensSubsystem->EvaluateTokenString(TreeItem->Path.Template, FNamingTokenFilterArgs({ UCineAssemblyNamingTokens::TokenNamespace }));
-	TreeItem->Path.Resolved = Result.EvaluatedText;
+	// Evaluate the token template string for this tree item
+	TreeItem->Path.Resolved = UCineAssemblyNamingTokens::GetResolvedText(TreeItem->Path.Template, CineAssemblyToConfigure);
 
 	// Evaluate the tokens for all of the child assets, then resort them alphabetically based on the resolved paths
 	for (const TSharedPtr<FHierarchyTreeItem>& Asset : TreeItem->ChildAssets)
@@ -622,57 +593,190 @@ TSharedRef<SWidget> SCineAssemblyConfigWindow::MakeNotesWidget()
 
 TSharedRef<SWidget> SCineAssemblyConfigWindow::MakeButtonsPanel()
 {
-	if (Mode == ECineAssemblyConfigMode::CreateNew)
-	{
-		return SNew(SBorder)
-			.BorderImage(FAppStyle::Get().GetBrush("Brushes.Panel"))
-			.Padding(16.0f)
-			.HAlign(HAlign_Right)
-			[
-				SNew(SHorizontalBox)
-
-				+ SHorizontalBox::Slot()
-					.MinWidth(118.0f)
-					.MaxWidth(118.0f)
-					.VAlign(VAlign_Center)
-					.Padding(0.0f, 0.0f, 8.0f, 0.0f)
-					[
-						SNew(SButton)
-							.Text(LOCTEXT("CreateAssetButton", "Create Asset"))
-							.ButtonStyle(FAppStyle::Get(), "PrimaryButton")
-							.HAlign(HAlign_Center)
-							.OnClicked(this, &SCineAssemblyConfigWindow::OnCreateAssetClicked)
-					]
-
-				+ SHorizontalBox::Slot()
-					.MinWidth(118.0f)
-					.MaxWidth(118.0f)
-					.VAlign(VAlign_Center)
-					[
-						SNew(SButton)
-							.Text(LOCTEXT("CancelButton", "Cancel"))
-							.HAlign(HAlign_Center)
-							.OnClicked(this, &SCineAssemblyConfigWindow::OnCancelClicked)
-					]
-			];
-	}
+	return SNew(SBorder)
+		.BorderImage(FAppStyle::Get().GetBrush("Brushes.Panel"))
+		.Padding(16.0f)
+		.HAlign(HAlign_Right)
+		[
+			SNew(SHorizontalBox)
 
-	return SNullWidget::NullWidget;
+			+ SHorizontalBox::Slot()
+				.MinWidth(118.0f)
+				.MaxWidth(118.0f)
+				.VAlign(VAlign_Center)
+				.Padding(0.0f, 0.0f, 8.0f, 0.0f)
+				[
+					SNew(SButton)
+						.Text(LOCTEXT("CreateAssetButton", "Create Asset"))
+						.ButtonStyle(FAppStyle::Get(), "PrimaryButton")
+						.HAlign(HAlign_Center)
+						.OnClicked(this, &SCineAssemblyConfigWindow::OnCreateAssetClicked)
+				]
+
+			+ SHorizontalBox::Slot()
+				.MinWidth(118.0f)
+				.MaxWidth(118.0f)
+				.VAlign(VAlign_Center)
+				[
+					SNew(SButton)
+						.Text(LOCTEXT("CancelButton", "Cancel"))
+						.HAlign(HAlign_Center)
+						.OnClicked(this, &SCineAssemblyConfigWindow::OnCancelClicked)
+				]
+		];
 }
 
 FReply SCineAssemblyConfigWindow::OnCreateAssetClicked()
 {
-	bShouldCreateNewAsset = true;
+	UCineAssemblyFactory::CreateConfiguredAssembly(CineAssemblyToConfigure, CreateAssetPath);
+
 	RequestDestroyWindow();
 	return FReply::Handled();
 }
 
 FReply SCineAssemblyConfigWindow::OnCancelClicked()
 {
-	bShouldCreateNewAsset = false;
 	RequestDestroyWindow();
 	return FReply::Handled();
 }
 
+void SCineAssemblyEditWidget::Construct(const FArguments& InArgs, UCineAssembly* InAssembly)
+{
+	CineAssembly = InAssembly;
+
+	TabSwitcher = SNew(SWidgetSwitcher)
+		+ SWidgetSwitcher::Slot()
+		[MakeOverviewWidget()]
+
+		+ SWidgetSwitcher::Slot()
+		[ UE::CineAssemblyWidgets::Private::MakeDetailsWidget(InAssembly) ];
+
+	ChildSlot
+		[
+			SNew(SBorder)
+				.BorderImage(FShotsEditorStyle::Get().GetBrush("Borders.PanelNoBorder"))
+				.Padding(8.0f)
+				[
+					SNew(SScrollBox)
+					.Orientation(Orient_Vertical)
+
+					+SScrollBox::Slot()
+						.AutoSize()
+						.HAlign(HAlign_Center)
+						.Padding(0.0f, 0.0f, 0.0f, 8.0f)
+						[
+							SNew(SSegmentedControl<int32>)
+								.Value(0)
+								.OnValueChanged_Lambda([Switcher = TabSwitcher](int32 NewValue)
+									{
+										Switcher->SetActiveWidgetIndex(NewValue);
+									})
+
+								+ SSegmentedControl<int32>::Slot(0)
+								.Text(LOCTEXT("OverviewTab", "Overview"))
+								.Icon(FSlateIcon(FShotsEditorStyle::StyleName, "Icons.Animation").GetIcon())
+
+								+ SSegmentedControl<int32>::Slot(1)
+								.Text(LOCTEXT("DetailsTab", "Details"))
+								.Icon(FSlateIcon(FAppStyle::GetAppStyleSetName(), "Icons.Details").GetIcon())
+						]
+
+					+ SScrollBox::Slot()
+						.FillSize(1.0f)
+						.HAlign(HAlign_Fill)
+						[
+							TabSwitcher.ToSharedRef()
+						]
+				]
+		];
+}
+
+TSharedRef<SWidget> SCineAssemblyEditWidget::MakeOverviewWidget()
+{
+	const UCineAssemblySchema* Schema = CineAssembly->GetSchema();
+
+	auto GetVisbility = [Schema]() -> EVisibility
+		{
+			return Schema ? EVisibility::Visible : EVisibility::Collapsed;
+		};
+
+	return SNew(SVerticalBox)
+
+		// TODO: Populate this with an actual preview image sourced from the schema
+		+ SVerticalBox::Slot()
+		.Padding(0.0f, 0.0f, 0.0f, 16.0f)
+		.MinHeight(300.0f)
+		.MaxHeight(300.0f)
+		[
+			SNew(SBorder)
+				.BorderImage(FShotsEditorStyle::Get().GetBrush("Borders.Background"))
+				.HAlign(HAlign_Center)
+				.VAlign(VAlign_Center)
+				[
+					SNew(STextBlock).Text(LOCTEXT("PreviewImageText", "PREVIEW IMAGE"))
+				]
+		]
+
+		+ SVerticalBox::Slot()
+		.AutoHeight()
+		.Padding(0.0f, 0.0f, 0.0f, 4.0f)
+		[
+			SNew(STextBlock)
+				.AutoWrapText(true)
+				.Text_Lambda([Schema]() -> FText
+					{
+						return Schema ? FText::FromString(Schema->SchemaName) : FText::GetEmpty();
+					})
+				.Visibility_Lambda(GetVisbility)
+		]
+
+		+ SVerticalBox::Slot()
+		.AutoHeight()
+		.Padding(0.0f, 0.0f, 0.0f, 8.0f)
+		[
+			SNew(STextBlock)
+				.Text(LOCTEXT("SchemClassName", "Cine Assembly Schema"))
+				.ColorAndOpacity(FSlateColor::UseSubduedForeground())
+				.Visibility_Lambda(GetVisbility)
+		]
+
+		+ SVerticalBox::Slot()
+		.AutoHeight()
+		.Padding(0.0f, 0.0f, 0.0f, 16.0f)
+		[
+			SNew(STextBlock)
+				.Text_Lambda([Schema]() -> FText
+					{
+						if (Schema)
+						{
+							return !Schema->Description.IsEmpty() ? FText::FromString(Schema->Description) : LOCTEXT("EmptyDescription", "No description");
+						}
+						return FText::GetEmpty();
+					})
+				.Visibility_Lambda(GetVisbility)
+		]
+
+		+ SVerticalBox::Slot()
+		.MinHeight(300.0f)
+		.FillContentHeight(1.0f)
+		[
+			SNew(SBorder)
+				.BorderImage(FShotsEditorStyle::Get().GetBrush("Borders.Background"))
+				.Padding(16.0f)
+				[
+					SNew(SMultiLineEditableText)
+						.HintText(LOCTEXT("NoteHintText", "What initial notes do you have about this shot?\nYou can always add more later."))
+						.Text_Lambda([this]()
+							{
+								return FText::FromString(CineAssembly->AssemblyNote);
+							})
+						.OnTextCommitted_Lambda([this](const FText& InText, ETextCommit::Type InCommitType)
+							{
+								CineAssembly->Modify();
+								CineAssembly->AssemblyNote = InText.ToString();
+							})
+				]
+		];
+}
 
 #undef LOCTEXT_NAMESPACE
diff --git a/Engine/Plugins/VirtualProduction/CinematicAssemblyTools/Source/ShotsEditor/Private/UI/CineAssembly/SCineAssemblyConfigWindow.h b/Engine/Plugins/VirtualProduction/CinematicAssemblyTools/Source/ShotsEditor/Private/UI/CineAssembly/SCineAssemblyConfigWindow.h
--- a/Engine/Plugins/VirtualProduction/CinematicAssemblyTools/Source/ShotsEditor/Private/UI/CineAssembly/SCineAssemblyConfigWindow.h
+++ b/Engine/Plugins/VirtualProduction/CinematicAssemblyTools/Source/ShotsEditor/Private/UI/CineAssembly/SCineAssemblyConfigWindow.h
@@ -11,13 +11,6 @@
 #include "Widgets/SWindow.h"
 #include "Widgets/Views/STreeView.h"
 
-/** Mode that configures the UI based on the intended user interactions with the Cine Assembly asset */
-enum class ECineAssemblyConfigMode : uint8
-{
-	CreateNew,
-	Edit
-};
-
 /** An entry in the hierarchy tree view */
 struct FHierarchyTreeItem
 {
@@ -52,10 +45,7 @@ public:
 	SLATE_BEGIN_ARGS(SCineAssemblyConfigWindow) {}
 	SLATE_END_ARGS()
 
-	void Construct(const FArguments& InArgs, UCineAssembly* InCineAssembly, ECineAssemblyConfigMode InConfigMode);
-
-	/** Whether a new asset should be created after closing this window */
-	bool ShouldCreateNewAsset() { return bShouldCreateNewAsset; }
+	void Construct(const FArguments& InArgs, const FString& InCreateAssetPath);
 
 private:
 	/** Creates the panel that displays the available Cine Assembly templates */
@@ -70,9 +60,6 @@ private:
 	/** Creates the widget to display for the Overview tab */
 	TSharedRef<SWidget> MakeOverviewWidget();
 
-	/** Creates the widget to display for the Details tab */
-	TSharedRef<SWidget> MakeDetailsWidget();
-
 	/** Creates the widget to display for the Hierarchy tab */
 	TSharedRef<SWidget> MakeHierarchyWidget();
 
@@ -112,9 +99,6 @@ private:
 	/** Recursively evaluates the tokens in the path of each item in the tree view */
 	void EvaluateHierarchyTokensRecursive(TSharedPtr<FHierarchyTreeItem> TreeItem);
 
-	/** Callback when this window is about to be closed */
-	void OnWindowClosed(const TSharedRef<SWindow>& InWindow);
-
 private:
 	/** Switcher that controls which tab widget is currently visible */
 	TSharedPtr<SWidgetSwitcher> TabSwitcher;
@@ -122,17 +106,14 @@ private:
 	/** Details View displaying the reflected properties of the Cine Assembly being configured */ 
 	TSharedPtr<IDetailsView> DetailsView;
 
-	/** Whether a new asset should be created after closing this window */
-	bool bShouldCreateNewAsset = false;
-
 	/** Transient object used only by this UI to configure the properties of the new asset that will get created by the Factory */
 	UCineAssembly* CineAssemblyToConfigure = nullptr;
 
 	/** The currently selected schema to use as a base for configuring the new cine assembly */
 	UCineAssemblySchema* SelectedSchema = nullptr;
 
-	/** Mode that configures the UI based on the intended user interactions with the Cine Assembly asset */
-	ECineAssemblyConfigMode Mode = ECineAssemblyConfigMode::CreateNew;
+	/** The root path where the configured assembly will be created */
+	FString CreateAssetPath;
 
 	/** The last time the naming tokens were updated */
 	FDateTime LastTokenUpdateTime;
@@ -146,3 +127,26 @@ private:
 	/** A read-only tree view of folders and assets that will be created based on the selected schema */
 	TSharedPtr<STreeView<TSharedPtr<FHierarchyTreeItem>>> HierarchyTreeView;
 };
+
+/** A panel that displays properties of a Cine Assembly asset */
+class SCineAssemblyEditWidget : public SCompoundWidget
+{
+public:
+	SCineAssemblyEditWidget() = default;
+
+	SLATE_BEGIN_ARGS(SCineAssemblyEditWidget) {}
+	SLATE_END_ARGS()
+
+	void Construct(const FArguments& InArgs, UCineAssembly* InAssembly);
+
+private:
+	/** Creates the widget to display for the Overview tab */
+	TSharedRef<SWidget> MakeOverviewWidget();
+
+private:
+	/** Switcher that controls which tab widget is currently visible */
+	TSharedPtr<SWidgetSwitcher> TabSwitcher;
+
+	/** Transient object used only by this UI to configure the properties of the new asset that will get created by the Factory */
+	UCineAssembly* CineAssembly = nullptr;
+};
\ No newline at end of file
