ShotManagement: Fix blank thumbnail on newly created Cine Assembly Schema assets.
#rnx

diff --git a/Engine/Plugins/VirtualProduction/CinematicAssemblyTools/Source/ShotsEditor/Private/CineAssemblySchema.cpp b/Engine/Plugins/VirtualProduction/CinematicAssemblyTools/Source/ShotsEditor/Private/CineAssemblySchema.cpp
--- a/Engine/Plugins/VirtualProduction/CinematicAssemblyTools/Source/ShotsEditor/Private/CineAssemblySchema.cpp
+++ b/Engine/Plugins/VirtualProduction/CinematicAssemblyTools/Source/ShotsEditor/Private/CineAssemblySchema.cpp
@@ -10,8 +10,12 @@ const FSoftObjectPath UCineAssemblySchema::DefaultThumbnailPath = FSoftObjectPat
 
 UCineAssemblySchema::UCineAssemblySchema()
 {
-	ThumbnailBrush = MakeShared<FSlateBrush>();
-	ThumbnailBrush->ImageSize = FVector2D(64.0f, 64.0f);
+	if (!HasAnyFlags(RF_ArchetypeObject | RF_ClassDefaultObject))
+	{
+		ThumbnailBrush = MakeShared<FSlateBrush>();
+		ThumbnailBrush->ImageSize = FVector2D(64.0f, 64.0f);
+		UpdateThumbnailBrush();
+	}
 }
 
 void UCineAssemblySchema::Serialize(FArchive& Ar)
