ShotManagement: Fix schema creation window to not require the use of a modal window.
The Create Schema window is no longer a modal window. It not creates the transient Cine Assembly Schema to be configured.

#rnx

diff --git a/Engine/Plugins/VirtualProduction/CinematicAssemblyTools/Source/ShotsEditor/Private/AssetDefinition_CineAssemblySchema.cpp b/Engine/Plugins/VirtualProduction/CinematicAssemblyTools/Source/ShotsEditor/Private/AssetDefinition_CineAssemblySchema.cpp
--- a/Engine/Plugins/VirtualProduction/CinematicAssemblyTools/Source/ShotsEditor/Private/AssetDefinition_CineAssemblySchema.cpp
+++ b/Engine/Plugins/VirtualProduction/CinematicAssemblyTools/Source/ShotsEditor/Private/AssetDefinition_CineAssemblySchema.cpp
@@ -4,8 +4,6 @@
 
 #include "CineAssemblySchema.h"
 #include "Framework/Application/SlateApplication.h"
-#include "Interfaces/IMainFrameModule.h"
-#include "Modules/ModuleManager.h"
 #include "UI/CineAssembly/SCineAssemblySchemaWindow.h"
 
 #define LOCTEXT_NAMESPACE "AssetTypeActions"
@@ -62,11 +60,8 @@ EAssetCommandResult UAssetDefinition_CineAssemblySchema::OpenAssets(const FAsset
 		const FAssetData& SchemaAssetData = OpenArgs.Assets[0];
 		if (UCineAssemblySchema* CineAssemblySchema = Cast<UCineAssemblySchema>(SchemaAssetData.GetAsset()))
 		{
-			IMainFrameModule& MainFrame = FModuleManager::LoadModuleChecked<IMainFrameModule>("MainFrame");
-			TSharedPtr<SWindow>	ParentWindow = MainFrame.GetParentWindow();
-
 			TSharedRef<SCineAssemblySchemaWindow> CineAssemblySchemaWindow = SNew(SCineAssemblySchemaWindow, Cast<UCineAssemblySchema>(CineAssemblySchema), ESchemaConfigMode::Edit);
-			FSlateApplication::Get().AddModalWindow(CineAssemblySchemaWindow, ParentWindow);
+			FSlateApplication::Get().AddWindow(CineAssemblySchemaWindow);
 		}
 	}
 
diff --git a/Engine/Plugins/VirtualProduction/CinematicAssemblyTools/Source/ShotsEditor/Private/CineAssemblySchemaFactory.cpp b/Engine/Plugins/VirtualProduction/CinematicAssemblyTools/Source/ShotsEditor/Private/CineAssemblySchemaFactory.cpp
--- a/Engine/Plugins/VirtualProduction/CinematicAssemblyTools/Source/ShotsEditor/Private/CineAssemblySchemaFactory.cpp
+++ b/Engine/Plugins/VirtualProduction/CinematicAssemblyTools/Source/ShotsEditor/Private/CineAssemblySchemaFactory.cpp
@@ -3,19 +3,18 @@
 #include "CineAssemblySchemaFactory.h"
 
 #include "AssetToolsModule.h"
+#include "CineAssemblySchema.h"
+#include "ContentBrowserModule.h"
 #include "Framework/Application/SlateApplication.h"
-#include "Interfaces/IMainFrameModule.h"
+#include "IContentBrowserSingleton.h"
 #include "Modules/ModuleManager.h"
 #include "UI/CineAssembly/SCineAssemblySchemaWindow.h"
-#include "UObject/Package.h"
 
 UCineAssemblySchemaFactory::UCineAssemblySchemaFactory()
 {
 	SupportedClass = UCineAssemblySchema::StaticClass();
 	bCreateNew = true;
 	bEditAfterNew = true;
-
-	SchemaToConfigure = nullptr;
 }
 
 bool UCineAssemblySchemaFactory::CanCreateNew() const
@@ -23,53 +22,56 @@ bool UCineAssemblySchemaFactory::CanCreateNew() const
 	return true;
 }
 
-UObject* UCineAssemblySchemaFactory::FactoryCreateNew(UClass* Class, UObject* InParent, FName Name, EObjectFlags Flags, UObject* Context, FFeedbackContext* Warn)
+bool UCineAssemblySchemaFactory::ConfigureProperties()
 {
-	if (SchemaToConfigure)
-	{
-		// If the user configured a name for the schema, use that as the name for the new asset. Otherwise, use the input default name
-		const FString NewSchemaName = SchemaToConfigure->SchemaName.IsEmpty() ? Name.ToString() : SchemaToConfigure->SchemaName;
+	IContentBrowserSingleton& ContentBrowser = FModuleManager::LoadModuleChecked<FContentBrowserModule>("ContentBrowser").Get();
+	const FString CurrentPath = ContentBrowser.GetCurrentPath().GetInternalPathString();
 
-		// Update the name, outer package, and flags of the configuration asset to no longer be transient
-		SchemaToConfigure->Rename(*NewSchemaName, InParent);
-		SchemaToConfigure->SetFlags(Flags);
-		SchemaToConfigure->ClearFlags(RF_Transient);
+	TSharedRef<SCineAssemblySchemaWindow> CineAssemblySchemaWindow = SNew(SCineAssemblySchemaWindow, CurrentPath);
+	FSlateApplication::Get().AddWindow(CineAssemblySchemaWindow);
 
-		// Unregister the temporary default name that we registered earlier for this schema instance
-		IAssetTools& AssetTools = FModuleManager::GetModuleChecked<FAssetToolsModule>("AssetTools").Get();
-		AssetTools.UnregisterDefaultAssetNameForClass(UCineAssemblySchema::StaticClass());
-	}
-
-	return SchemaToConfigure;
+	return false;
 }
 
-bool UCineAssemblySchemaFactory::ConfigureProperties()
+UObject* UCineAssemblySchemaFactory::FactoryCreateNew(UClass* Class, UObject* InParent, FName Name, EObjectFlags Flags, UObject* Context, FFeedbackContext* Warn)
 {
-	// Create a new transient CineAssembly to configure in the UI.
-	// If the configuration is successful, this will turn into the persistent object created by this factory.
-	SchemaToConfigure = NewObject<UCineAssemblySchema>(GetTransientPackage(), NAME_None, RF_Transient);
-
-	// Spawn a new window to configure the properties of the asset being created
-	IMainFrameModule& MainFrame = FModuleManager::LoadModuleChecked<IMainFrameModule>("MainFrame");
-	TSharedPtr<SWindow>	ParentWindow = MainFrame.GetParentWindow();
-
-	TSharedRef<SCineAssemblySchemaWindow> CineAssemblySchemaWindow = SNew(SCineAssemblySchemaWindow, SchemaToConfigure, ESchemaConfigMode::CreateNew);
-	FSlateApplication::Get().AddModalWindow(CineAssemblySchemaWindow, ParentWindow);
+	UCineAssemblySchema* NewSchema = NewObject<UCineAssemblySchema>(InParent, Name, Flags);
+	return NewSchema;
+}
 
-	if (CineAssemblySchemaWindow->ShouldCreateNewAsset() == false)
+void UCineAssemblySchemaFactory::CreateConfiguredSchema(UCineAssemblySchema* ConfiguredSchema, const FString& CreateAssetPath)
+{
+	// The intended use of this function is to take a pre-configured, transient assembly, create a valid package for it, and initialize it.
+	if (ConfiguredSchema->GetPackage() != GetTransientPackage())
 	{
-		return false;
+		return;
 	}
 
-	// One of the properties configured by the user in the schema window is the schema name.
-	// We want the actual asset name to match the value of this property, and we don't want the user to have to retype it.
-	// This workaround allows us to temporarily set a default name for the schema asset. 
-	// After FactoryCreateNew() is finished, we will unregister this default name override. 
-	if (!SchemaToConfigure->SchemaName.IsEmpty())
+	// If the assembly name is empty, assign it a valid default name
+	if (ConfiguredSchema->SchemaName.IsEmpty())
 	{
-		IAssetTools& AssetTools = FModuleManager::GetModuleChecked<FAssetToolsModule>("AssetTools").Get();
-		AssetTools.RegisterDefaultAssetNameForClass(UCineAssemblySchema::StaticClass(), SchemaToConfigure->SchemaName);
+		ConfiguredSchema->SchemaName = TEXT("NewCineAssemblySchema");
 	}
 
-	return true;
+	const FString DesiredPackageName = CreateAssetPath / ConfiguredSchema->SchemaName;
+	const FString DesiredSuffix = TEXT("");
+
+	IAssetTools& AssetTools = FModuleManager::GetModuleChecked<FAssetToolsModule>("AssetTools").Get();
+
+	FString UniquePackageName;
+	FString UniqueAssetName;
+	AssetTools.CreateUniqueAssetName(DesiredPackageName, DesiredSuffix, UniquePackageName, UniqueAssetName);
+
+	// The input schema object was created in the transient package while its properties were configured.
+	// Now, we can create a real package for it, rename it, and update its object flags.
+	UPackage* Package = CreatePackage(*UniquePackageName);
+	ConfiguredSchema->Rename(*UniqueAssetName, Package);
+
+	const EObjectFlags Flags = RF_Public | RF_Standalone | RF_Transactional;
+	ConfiguredSchema->SetFlags(Flags);
+	ConfiguredSchema->ClearFlags(RF_Transient);
+
+	// Refresh the content browser to make any new assets and folders are immediately visible
+	IContentBrowserSingleton& ContentBrowser = FModuleManager::LoadModuleChecked<FContentBrowserModule>("ContentBrowser").Get();
+	ContentBrowser.SetSelectedPaths({ CreateAssetPath }, true);
 }
diff --git a/Engine/Plugins/VirtualProduction/CinematicAssemblyTools/Source/ShotsEditor/Private/CineAssemblySchemaFactory.h b/Engine/Plugins/VirtualProduction/CinematicAssemblyTools/Source/ShotsEditor/Private/CineAssemblySchemaFactory.h
--- a/Engine/Plugins/VirtualProduction/CinematicAssemblyTools/Source/ShotsEditor/Private/CineAssemblySchemaFactory.h
+++ b/Engine/Plugins/VirtualProduction/CinematicAssemblyTools/Source/ShotsEditor/Private/CineAssemblySchemaFactory.h
@@ -3,10 +3,11 @@
 #pragma once
 
 #include "Factories/Factory.h"
-#include "CineAssemblySchema.h"
 
 #include "CineAssemblySchemaFactory.generated.h"
 
+class UCineAssemblySchema;
+
 /**
  * Factory class used to create new UCineAssemblySchema objects
  * Before creating a new assembly schema, the factory will spawn a window to configure the properties of the asset that is being created
@@ -16,6 +17,10 @@ class UCineAssemblySchemaFactory : public UFactory
 {
 	GENERATED_BODY()
 
+public:
+	/** Takes a pre-configured, transient schema, creates a valid package for it, and initializes it */
+	static void CreateConfiguredSchema(UCineAssemblySchema* ConfiguredSchema, const FString& CreateAssetPath);
+
 protected:
 	UCineAssemblySchemaFactory();
 
@@ -24,8 +29,4 @@ protected:
 	virtual bool CanCreateNew() const override;
 	virtual bool ConfigureProperties() override;
 	// End UFactory Interface
-
-private:
-	/** The schema object that will be created by this factory and configured by the user */
-	UCineAssemblySchema* SchemaToConfigure;
 };
diff --git a/Engine/Plugins/VirtualProduction/CinematicAssemblyTools/Source/ShotsEditor/Private/UI/CineAssembly/SCineAssemblySchemaWindow.cpp b/Engine/Plugins/VirtualProduction/CinematicAssemblyTools/Source/ShotsEditor/Private/UI/CineAssembly/SCineAssemblySchemaWindow.cpp
--- a/Engine/Plugins/VirtualProduction/CinematicAssemblyTools/Source/ShotsEditor/Private/UI/CineAssembly/SCineAssemblySchemaWindow.cpp
+++ b/Engine/Plugins/VirtualProduction/CinematicAssemblyTools/Source/ShotsEditor/Private/UI/CineAssembly/SCineAssemblySchemaWindow.cpp
@@ -5,6 +5,7 @@
 #include "AssetRegistry/AssetRegistryModule.h"
 #include "AssetViewUtils.h"
 #include "CineAssemblyNamingTokens.h"
+#include "CineAssemblySchemaFactory.h"
 #include "Framework/MultiBox/MultiBoxBuilder.h"
 #include "IStructureDetailsView.h"
 #include "Modules/ModuleManager.h"
@@ -13,6 +14,7 @@
 #include "PropertyEditorModule.h"
 #include "ShotsEditorStyle.h"
 #include "Styling/StyleColors.h"
+#include "UObject/Package.h"
 #include "Widgets/Input/SButton.h"
 #include "Widgets/Input/SEditableTextBox.h"
 #include "Widgets/Input/SSegmentedControl.h"
@@ -29,6 +31,16 @@ namespace UE::CineAssemblyTools::Private
 	auto SortTreeItems = [](const TSharedPtr<FSchemaTreeItem>& A, const TSharedPtr<FSchemaTreeItem>& B) { return A->Path < B->Path; };
 };
 
+void SCineAssemblySchemaWindow::Construct(const FArguments& InArgs, const FString& InCreateAssetPath)
+{
+	// Create a new transient CineAssemblySchema to configure in the UI.
+	// If the configuration is successful, this will turn into the persistent object created by the factory.
+	SchemaToConfigure = NewObject<UCineAssemblySchema>(GetTransientPackage(), NAME_None, RF_Transient);
+	CreateAssetPath = InCreateAssetPath;
+
+	Construct(InArgs, SchemaToConfigure, ESchemaConfigMode::CreateNew);
+}
+
 void SCineAssemblySchemaWindow::Construct(const FArguments& InArgs, UCineAssemblySchema* InSchema, ESchemaConfigMode InConfigMode)
 {
 	SchemaToConfigure = InSchema;
@@ -1102,14 +1114,14 @@ void SCineAssemblySchemaWindow::OnSchemaPropertiesChanged(const FPropertyChanged
 
 FReply SCineAssemblySchemaWindow::OnCreateAssetClicked()
 {
-	bShouldCreateNewAsset = true;
+	UCineAssemblySchemaFactory::CreateConfiguredSchema(SchemaToConfigure, CreateAssetPath);
+
 	RequestDestroyWindow();
 	return FReply::Handled();
 }
 
 FReply SCineAssemblySchemaWindow::OnCancelClicked()
 {
-	bShouldCreateNewAsset = false;
 	RequestDestroyWindow();
 	return FReply::Handled();
 }
diff --git a/Engine/Plugins/VirtualProduction/CinematicAssemblyTools/Source/ShotsEditor/Private/UI/CineAssembly/SCineAssemblySchemaWindow.h b/Engine/Plugins/VirtualProduction/CinematicAssemblyTools/Source/ShotsEditor/Private/UI/CineAssembly/SCineAssemblySchemaWindow.h
--- a/Engine/Plugins/VirtualProduction/CinematicAssemblyTools/Source/ShotsEditor/Private/UI/CineAssembly/SCineAssemblySchemaWindow.h
+++ b/Engine/Plugins/VirtualProduction/CinematicAssemblyTools/Source/ShotsEditor/Private/UI/CineAssembly/SCineAssemblySchemaWindow.h
@@ -98,9 +98,7 @@ public:
 	SLATE_END_ARGS()
 
 	void Construct(const FArguments& InArgs, UCineAssemblySchema* InSchema, ESchemaConfigMode InConfigMode);
-
-	/** Whether a new asset should be created after closing this window */
-	bool ShouldCreateNewAsset() { return bShouldCreateNewAsset; }
+	void Construct(const FArguments& InArgs, const FString& InCreateAssetPath);
 
 private:
 	/** Creates the panel that displays the tab menu */
@@ -224,15 +222,15 @@ private:
 	/** Switcher that controls which tab widget in the Asset Tab is currently visible */
 	TSharedPtr<SWidgetSwitcher> AssetTabSwitcher;
 
-	/** Whether a new asset should be created after closing this window */
-	bool bShouldCreateNewAsset = false;
-
 	/** Transient object used only by this UI to configure the properties of the new asset that will get created by the factory */
 	UCineAssemblySchema* SchemaToConfigure = nullptr;
 
 	/** Mode that configures the UI based on the intended user interactions with the Cine Assembly schema asset */
 	ESchemaConfigMode Mode = ESchemaConfigMode::CreateNew;
 
+	/** The root path where the configured schema will be created */
+	FString CreateAssetPath;
+
 	/** Whether the schema asset being configured needs to be renamed */
 	bool bPendingRename = false;
 
