ShotManagement: Fix Cine Assembly hierarchy view not properly updating when the schema is changed.
#rnx

diff --git a/Engine/Plugins/VirtualProduction/CinematicAssemblyTools/Source/ShotsEditor/Private/UI/CineAssembly/SCineAssemblyConfigWindow.cpp b/Engine/Plugins/VirtualProduction/CinematicAssemblyTools/Source/ShotsEditor/Private/UI/CineAssembly/SCineAssemblyConfigWindow.cpp
--- a/Engine/Plugins/VirtualProduction/CinematicAssemblyTools/Source/ShotsEditor/Private/UI/CineAssembly/SCineAssemblyConfigWindow.cpp
+++ b/Engine/Plugins/VirtualProduction/CinematicAssemblyTools/Source/ShotsEditor/Private/UI/CineAssembly/SCineAssemblyConfigWindow.cpp
@@ -117,8 +117,8 @@ TSharedRef<SWidget> SCineAssemblyConfigWindow::MakeCineTemplatePanel()
 	Config.bAllowRename = false;
 	Config.bCanShowClasses = false;
 	Config.bCanShowFolders = false;
-	Config.bForceShowEngineContent = true;
-	Config.bForceShowPluginContent = true;
+	Config.bForceShowEngineContent = false;
+	Config.bForceShowPluginContent = false;
 	Config.bAddFilterUI = false;
 	Config.bShowPathInColumnView = true;
 	Config.OnAssetSelected = FOnAssetSelected::CreateSP(this, &SCineAssemblyConfigWindow::OnSchemaSelected);
@@ -405,6 +405,8 @@ void SCineAssemblyConfigWindow::PopulateHierarchyTree()
 {
 	HierarchyTreeItems.Empty();
 	HierarchyTreeItems.Add(RootItem);
+	RootItem->ChildAssets.Empty();
+	RootItem->ChildFolders.Empty();
 
 	auto AddItemsToTree = [this](TArray<FTemplateString>& ItemList, FHierarchyTreeItem::EItemType ItemType)
 		{
