ShotManagement: Update instructional text on Cine Assembly Schema hierarchy tab
#rnx

diff --git a/Engine/Plugins/VirtualProduction/CinematicAssemblyTools/Source/ShotsEditor/Private/UI/CineAssembly/SCineAssemblySchemaWindow.cpp b/Engine/Plugins/VirtualProduction/CinematicAssemblyTools/Source/ShotsEditor/Private/UI/CineAssembly/SCineAssemblySchemaWindow.cpp
--- a/Engine/Plugins/VirtualProduction/CinematicAssemblyTools/Source/ShotsEditor/Private/UI/CineAssembly/SCineAssemblySchemaWindow.cpp
+++ b/Engine/Plugins/VirtualProduction/CinematicAssemblyTools/Source/ShotsEditor/Private/UI/CineAssembly/SCineAssemblySchemaWindow.cpp
@@ -399,35 +399,6 @@ TSharedRef<SWidget> SCineAssemblySchemaWindow::MakeHierarchyTabContent()
 
 	TSharedRef<SWidget> ContentHierarhcyWidget = SNew(SVerticalBox)
 
-		+ SVerticalBox::Slot()
-		.AutoHeight()
-		.Padding(0.0f, 0.0f, 0.0f, 4.0f)
-		[
-			SNew(STextBlock)
-				.Text(LOCTEXT("ContentHierarchyTitle", "Content Hierarchy"))
-				.Font(FShotsEditorStyle::Get().GetFontStyle("ProductionWizard.HeadingFont"))
-		]
-
-		+ SVerticalBox::Slot()
-		.AutoHeight()
-		.Padding(0.0f, 0.0f, 0.0f, 8.0f)
-		[
-			SNew(STextBlock)
-				.AutoWrapText(true)
-				.Text(LOCTEXT("ContentHierarchyText1", "Create a content hierarchy for your shot type. "
-					"This hierarchy will be automatically created when creating shots from this shot type."))
-		]
-
-		+ SVerticalBox::Slot()
-		.AutoHeight()
-		.Padding(0.0f, 0.0f, 0.0f, 16.0f)
-		[
-			SNew(STextBlock)
-				.AutoWrapText(true)
-				.Text(LOCTEXT("ContentHierarchyText2", "Assets being created or imported with this shot can be found to the right. "
-					"Any assets being created based on your timeline template must be assigned a save location before this shot type can be created."))
-		]
-
 		+ SVerticalBox::Slot()
 		.AutoHeight()
 		.HAlign(HAlign_Left)
@@ -467,16 +438,6 @@ TSharedRef<SWidget> SCineAssemblySchemaWindow::MakeHierarchyTabContent()
 
 	TSharedRef<SWidget> AssetListWidget = SNew(SVerticalBox)
 
-		+ SVerticalBox::Slot()
-		.AutoHeight()
-		.Padding(0.0f, 0.0f, 0.0f, 16.0f)
-		[
-			SNew(STextBlock)
-				.AutoWrapText(true)
-				.Text(LOCTEXT("AssetListInstructions", "The following assets will be created based on your Timeline Template. "
-				"Drag and drop them to a location in the hierarchy to determine where they should be saved when a new shot is created from this shot type"))
-		]
-
 		+ SVerticalBox::Slot()
 		.AutoHeight()
 		.HAlign(HAlign_Left)
@@ -521,28 +482,58 @@ TSharedRef<SWidget> SCineAssemblySchemaWindow::MakeHierarchyTabContent()
 				]
 		];
 
-	return SNew(SHorizontalBox)
-
-		+ SHorizontalBox::Slot()
-		.FillWidth(0.5f)
+	return SNew(SBorder)
 		.Padding(16.0f)
+		.BorderImage(FShotsEditorStyle::Get().GetBrush("ProductionWizard.PanelBackground"))
 		[
-			AssetListWidget
-		]
+			SNew(SVerticalBox)
 
-		+ SHorizontalBox::Slot()
-		.AutoWidth()
-		[
-			SNew(SSeparator)
-				.Orientation(Orient_Vertical)
-				.Thickness(2.0f)
-		]
+			+ SVerticalBox::Slot()
+				.AutoHeight()
+				.Padding(0.0f, 0.0f, 0.0f, 8.0f)
+				[
+					SNew(STextBlock)
+						.Text(LOCTEXT("ContentHierarchyTitle", "Content Hierarchy"))
+						.Font(FShotsEditorStyle::Get().GetFontStyle("ProductionWizard.HeadingFont"))
+				]
 
-		+ SHorizontalBox::Slot()
-		.FillWidth(0.5f)
-		.Padding(16.0f)
-		[
-			ContentHierarhcyWidget
+			+ SVerticalBox::Slot()
+				.AutoHeight()
+				.Padding(0.0f, 0.0f, 0.0f, 24.0f)
+				[
+					SNew(STextBlock)
+						.AutoWrapText(true)
+						.Text(LOCTEXT("AssetListInstructions", "Add new named subsequences to the list on the left, then drag and drop them to a location in the folder tree on the right. "
+							"When a new Cine Assembly is created using this Schema, the subsequences will automatically be created and added as tracks to the Assembly."))
+				]
+
+			+ SVerticalBox::Slot()
+				.FillHeight(1.0f)
+				[
+					SNew(SHorizontalBox)
+
+					+ SHorizontalBox::Slot()
+						.FillWidth(0.5f)
+						.Padding(0.0f, 0.0f, 16.0f, 0.0f)
+						[
+							AssetListWidget
+						]
+
+					+ SHorizontalBox::Slot()
+						.AutoWidth()
+						[
+							SNew(SSeparator)
+								.Orientation(Orient_Vertical)
+								.Thickness(2.0f)
+						]
+
+					+ SHorizontalBox::Slot()
+						.FillWidth(0.5f)
+						.Padding(16.0f, 0.0f, 0.0f, 0.0f)
+						[
+							ContentHierarhcyWidget
+						]
+				]
 		];
 }
 
