ShotManagement: Rename "ShotsEditor" module to "CineAssemblyToolsEditor" and replace all references to "Shots" with "CineAssemblyTools".
Add "CineAssemblyTools" Runtime module and move CineAssembly, CineAssemblySchema, and CineAssemblyNamingTokens classes to this new module.

#rnx

#virtualized

diff --git a/Engine/Plugins/VirtualProduction/CinematicAssemblyTools/CinematicAssemblyTools.uplugin b/Engine/Plugins/VirtualProduction/CinematicAssemblyTools/CinematicAssemblyTools.uplugin
--- a/Engine/Plugins/VirtualProduction/CinematicAssemblyTools/CinematicAssemblyTools.uplugin
+++ b/Engine/Plugins/VirtualProduction/CinematicAssemblyTools/CinematicAssemblyTools.uplugin
@@ -3,7 +3,7 @@
 	"Version": 1,
 	"VersionName": "1.0",
 	"FriendlyName": "Cinematic Assembly Tools",
-	"Description": "Shot management tools for creating and editing linear content",
+	"Description": "Cinematic pipeline tools for shot management and linear content creation",
 	"Category": "Virtual Production",
 	"CreatedBy": "Epic Games, Inc.",
 	"CreatedByURL": "https://epicgames.com",
@@ -15,12 +15,15 @@
 	"IsExperimentalVersion": true,
 	"Hidden": false,
 	"Installed": false,
-	"SupportedPrograms": [ "LiveLinkHub" ],
 	"Modules": [
 		{
-			"Name": "ShotsEditor",
+			"Name": "CineAssemblyTools",
+			"Type": "Runtime",
+			"LoadingPhase": "Default"
+		},
+		{
+			"Name": "CineAssemblyToolsEditor",
 			"Type": "Editor",
-			"ProgramAllowlist": [ "LiveLinkHub" ],
 			"LoadingPhase": "Default"
 		}
 	],
diff --git a/Engine/Plugins/VirtualProduction/CinematicAssemblyTools/Config/BaseCinematicAssemblyTools.ini b/Engine/Plugins/VirtualProduction/CinematicAssemblyTools/Config/BaseCinematicAssemblyTools.ini
new file mode 100644
--- /dev/null
+++ b/Engine/Plugins/VirtualProduction/CinematicAssemblyTools/Config/BaseCinematicAssemblyTools.ini
@@ -0,0 +1,4 @@
+[CoreRedirects]
+
++ClassRedirects=(OldName="/Script/ShotsEditor.CineAssembly", NewName="/Script/CineAssemblyTools.CineAssembly")
++ClassRedirects=(OldName="/Script/ShotsEditor.CineAssemblySchema", NewName="/Script/CineAssemblyTools.CineAssemblySchema")
\ No newline at end of file
diff --git a/Engine/Plugins/VirtualProduction/CinematicAssemblyTools/Source/CineAssemblyTools/Private/CineAssemblyToolsModule.cpp b/Engine/Plugins/VirtualProduction/CinematicAssemblyTools/Source/CineAssemblyTools/Private/CineAssemblyToolsModule.cpp
new file mode 100644
--- /dev/null
+++ b/Engine/Plugins/VirtualProduction/CinematicAssemblyTools/Source/CineAssemblyTools/Private/CineAssemblyToolsModule.cpp
@@ -0,0 +1,5 @@
+// Copyright Epic Games, Inc. All Rights Reserved.
+
+#include "Modules/ModuleManager.h"
+
+IMPLEMENT_MODULE(FDefaultModuleImpl, CineAssemblyTools);
diff --git a/Engine/Plugins/VirtualProduction/CinematicAssemblyTools/Source/CineAssemblyToolsEditor/Private/CineAssemblyToolsEditorModule.cpp b/Engine/Plugins/VirtualProduction/CinematicAssemblyTools/Source/CineAssemblyToolsEditor/Private/CineAssemblyToolsEditorModule.cpp
new file mode 100644
--- /dev/null
+++ b/Engine/Plugins/VirtualProduction/CinematicAssemblyTools/Source/CineAssemblyToolsEditor/Private/CineAssemblyToolsEditorModule.cpp
@@ -0,0 +1,199 @@
+// Copyright Epic Games, Inc. All Rights Reserved.
+
+#include "Modules/ModuleManager.h"
+
+#include "CineAssembly.h"
+#include "CineAssemblySchema.h"
+#include "CineAssemblyCustomization.h"
+#include "CineAssemblyMetadataCustomization.h"
+#include "CineAssemblyToolsStyle.h"
+#include "ContentBrowserMenuContexts.h"
+#include "Editor.h"
+#include "Framework/Application/SlateApplication.h"
+#include "Framework/Docking/TabManager.h"
+#include "Interfaces/IMainFrameModule.h"
+#include "MoviePipelineQueue.h"
+#include "ProductionSettings.h"
+#include "ProductionSettingsCustomization.h"
+#include "PropertyEditorModule.h"
+#include "Subsystems/AssetEditorSubsystem.h"
+#include "ToolMenus.h"
+#include "UI/SProductionWizard.h"
+#include "UI/CineAssembly/SCineAssemblyConfigWindow.h"
+#include "Widgets/Docking/SDockTab.h"
+#include "WorkspaceMenuStructure.h"
+#include "WorkspaceMenuStructureModule.h"
+#include "TakeRecorder/CineAssemblyTakeRecorderIntegration.h"
+
+#define LOCTEXT_NAMESPACE "CineAssemblyToolsEditorModule"
+
+static const FName ProductionWizardTabName("ProductionWizard");
+
+class FCineAssemblyToolsEditorModule : public IModuleInterface
+{
+public:
+
+	//~ Begin IModuleInterface interface
+	virtual void StartupModule() override;
+	virtual void ShutdownModule() override;
+	//~ End IModuleInterface interface
+
+private:
+	/** Manages all integration of Cinematic Assembly Tools with Take Recorder. */
+	TUniquePtr<FCineAssemblyTakeRecorderIntegration> TakeRecorderIntegration;
+	
+	TSharedRef<SDockTab> SpawnProductionWizard(const FSpawnTabArgs& SpawnTabArgs);
+	void RegisterMenus();
+
+	/** Called whenever a job has its sequence set. */
+	void OnSequenceSet(UMoviePipelineExecutorJob* Job, ULevelSequence* Sequence);
+
+	/** Opens a window to edit the properties of the input assembly */
+	void OpenAssemblyForEdit(UCineAssembly* Assembly);
+
+private:
+	FDelegateHandle OnSequenceSetHandle;
+};
+
+void FCineAssemblyToolsEditorModule::StartupModule()
+{
+	FCineAssemblyToolsStyle::Get();
+
+	FPropertyEditorModule& PropertyEditorModule = FModuleManager::Get().LoadModuleChecked<FPropertyEditorModule>("PropertyEditor");
+
+	PropertyEditorModule.RegisterCustomClassLayout(
+		UProductionSettings::StaticClass()->GetFName(),
+		FOnGetDetailCustomizationInstance::CreateStatic(&FProductionSettingsCustomization::MakeInstance)
+	);
+
+	PropertyEditorModule.RegisterCustomClassLayout(
+		UCineAssembly::StaticClass()->GetFName(),
+		FOnGetDetailCustomizationInstance::CreateStatic(&FCineAssemblyCustomization::MakeInstance)
+	);
+
+	PropertyEditorModule.RegisterCustomPropertyTypeLayout(
+		FAssemblyMetadataDesc::StaticStruct()->GetFName(),
+		FOnGetPropertyTypeCustomizationInstance::CreateStatic(&FCineAssemblyMetadataCustomization::MakeInstance)
+	);
+
+	// Add a new entry to the Tools->Cinematics menu to spawn the Production Wizard
+	FGlobalTabmanager::Get()->RegisterNomadTabSpawner(ProductionWizardTabName, FOnSpawnTab::CreateRaw(this, &FCineAssemblyToolsEditorModule::SpawnProductionWizard))
+		.SetGroup(WorkspaceMenu::GetMenuStructure().GetLevelEditorCinematicsCategory())
+		.SetTooltipText(LOCTEXT("ProductionWizardTabTooltip", "Open the Production Wizard"))
+		.SetDisplayName(LOCTEXT("ProductionWizardTabTitle", "Production Wizard"))
+		.SetIcon(FSlateIcon(FAppStyle::GetAppStyleSetName(), "ProjectSettings.TabIcon"));
+
+	RegisterMenus();
+
+	FCoreDelegates::OnPostEngineInit.AddLambda([this]() 
+		{
+			TakeRecorderIntegration = MakeUnique<FCineAssemblyTakeRecorderIntegration>();
+		});
+
+	OnSequenceSetHandle = UMoviePipelineExecutorJob::OnSequenceSetEvent().AddRaw(this, &FCineAssemblyToolsEditorModule::OnSequenceSet);
+}
+
+void FCineAssemblyToolsEditorModule::ShutdownModule()
+{
+	FCoreDelegates::OnPostEngineInit.RemoveAll(this);
+	TakeRecorderIntegration.Reset();
+
+	FPropertyEditorModule& PropertyModule = FModuleManager::GetModuleChecked<FPropertyEditorModule>("PropertyEditor");
+	PropertyModule.UnregisterCustomClassLayout(UProductionSettings::StaticClass()->GetFName());
+
+	FGlobalTabmanager::Get()->UnregisterNomadTabSpawner(ProductionWizardTabName);
+
+	UMoviePipelineExecutorJob::OnSequenceSetEvent().Remove(OnSequenceSetHandle);
+}
+
+TSharedRef<SDockTab> FCineAssemblyToolsEditorModule::SpawnProductionWizard(const FSpawnTabArgs& SpawnTabArgs)
+{
+	return SNew(SDockTab)
+		.TabRole(ETabRole::NomadTab)
+		[
+			SNew(SProductionWizard)
+		];
+}
+
+void FCineAssemblyToolsEditorModule::RegisterMenus()
+{
+	UToolMenus* ToolMenus = UToolMenus::Get();
+
+	// Add a Content Browser context menu action to Cine Assembly assets that allows the asset to be opened in Sequencer without loading its associated map
+	const FName ContextMenuName = TEXT("ContentBrowser.AssetContextMenu.CineAssembly");
+	if (UToolMenu* Menu = ToolMenus->ExtendMenu(ContextMenuName))
+	{
+		const FName AssetActionSectionName = TEXT("GetAssetActions");
+		FToolMenuSection& Section = Menu->FindOrAddSection(AssetActionSectionName);
+
+		Section.AddDynamicEntry("CineAssemblyActions", FNewToolMenuSectionDelegate::CreateLambda([this](FToolMenuSection& InSection)
+		{
+			if (UContentBrowserAssetContextMenuContext* Context = InSection.FindContext<UContentBrowserAssetContextMenuContext>())
+			{
+				if ((Context->SelectedAssets.Num() == 1) && Context->SelectedAssets[0].IsInstanceOf(UCineAssembly::StaticClass()))
+				{
+					const FAssetData& CineAssemblyData = Context->SelectedAssets[0];
+					UCineAssembly* CineAssembly = Cast<UCineAssembly>(CineAssemblyData.GetAsset());
+
+					InSection.AddMenuEntry(
+						"OpenInCurrentMapEntry",
+						LOCTEXT("OpenInCurrentMap", "Open in Current Map"),
+						LOCTEXT("OpenInCurrentMapTooltip", "Opens the level sequence in Sequencer but does not automatically load the associated map"),
+						FSlateIcon(FAppStyle::GetAppStyleSetName(), "LevelEditor.OpenCinematic"),
+						FUIAction(FExecuteAction::CreateLambda([CineAssembly]()
+						{
+							// Note: Normally, when a cine assembly is opened in the editor, its associated map is loaded during the PrepareToActiveAssets step.
+							// By directly opening the asset editor (Sequencer) for this asset, we can skip that map loading step. 
+							GEditor->GetEditorSubsystem<UAssetEditorSubsystem>()->OpenEditorForAssets({ CineAssembly });
+						})),
+						EUserInterfaceActionType::None
+					);
+
+					InSection.AddMenuEntry(
+						"OpenInAssetEditor",
+						LOCTEXT("OpenInAssetEditor", "Edit Properties"),
+						LOCTEXT("OpenInAssetEditorTooltip", "Edit Properties"),
+						FSlateIcon(FAppStyle::GetAppStyleSetName(), "Icons.Edit"),
+						FUIAction(FExecuteAction::CreateRaw(this, &FCineAssemblyToolsEditorModule::OpenAssemblyForEdit, CineAssembly)),
+						EUserInterfaceActionType::None
+					);
+				}
+			}
+		}));
+	}
+}
+
+void FCineAssemblyToolsEditorModule::OpenAssemblyForEdit(UCineAssembly* Assembly)
+{
+	TSharedRef<SCineAssemblyEditWidget> CineAssemblyWidget = SNew(SCineAssemblyEditWidget, Assembly);
+
+	const FVector2D DefaultWindowSize = FVector2D(600, 800);
+	TSharedRef<SWindow> EditWindow = SNew(SWindow)
+		.Title(LOCTEXT("EditCineAssemblyTitle", "Edit Cine Assembly"))
+		.SizingRule(ESizingRule::UserSized)
+		.ClientSize(DefaultWindowSize)
+		.SupportsMaximize(false)
+		.SupportsMinimize(false);
+
+	EditWindow->SetContent(CineAssemblyWidget);
+	FSlateApplication::Get().AddWindow(EditWindow);
+}
+
+void FCineAssemblyToolsEditorModule::OnSequenceSet(UMoviePipelineExecutorJob* Job, ULevelSequence* Sequence)
+{
+	if (Job)
+	{
+		if (const UCineAssembly* CineAssembly = Cast<UCineAssembly>(Sequence))
+		{
+			// For our cine assemblies, we want to always use the map associated with the sequence.
+			if (CineAssembly->Level.IsValid())
+			{
+				Job->Map = CineAssembly->Level;
+			}
+		}
+	}
+}
+
+IMPLEMENT_MODULE(FCineAssemblyToolsEditorModule, CineAssemblyToolsEditor);
+
+#undef LOCTEXT_NAMESPACE
