ShotManagement: Add support for naming tokens in the Production Wizard for default asset names and folder template names.
Adds a step in AssetTools::GetDefaultAssetNameForClass to sanitize the name found in the map, which gives us a hook to evaluate tokens before handing back the default asset name.
#rnx

diff --git a/Engine/Plugins/VirtualProduction/CinematicAssemblyTools/Source/CineAssemblyToolsEditor/Private/CineAssemblyToolsEditorModule.cpp b/Engine/Plugins/VirtualProduction/CinematicAssemblyTools/Source/CineAssemblyToolsEditor/Private/CineAssemblyToolsEditorModule.cpp
--- a/Engine/Plugins/VirtualProduction/CinematicAssemblyTools/Source/CineAssemblyToolsEditor/Private/CineAssemblyToolsEditorModule.cpp
+++ b/Engine/Plugins/VirtualProduction/CinematicAssemblyTools/Source/CineAssemblyToolsEditor/Private/CineAssemblyToolsEditorModule.cpp
@@ -2,6 +2,7 @@
 
 #include "Modules/ModuleManager.h"
 
+#include "AssetToolsModule.h"
 #include "CineAssembly.h"
 #include "CineAssemblySchema.h"
 #include "CineAssemblyCustomization.h"
@@ -13,6 +14,7 @@
 #include "Framework/Docking/TabManager.h"
 #include "Interfaces/IMainFrameModule.h"
 #include "MoviePipelineQueue.h"
+#include "NamingTokensEngineSubsystem.h"
 #include "ProductionSettings.h"
 #include "ProductionSettingsCustomization.h"
 #include "PropertyEditorModule.h"
@@ -39,9 +41,8 @@ public:
 	//~ End IModuleInterface interface
 
 private:
-	/** Manages all integration of Cinematic Assembly Tools with Take Recorder. */
-	TUniquePtr<FCineAssemblyTakeRecorderIntegration> TakeRecorderIntegration;
-	
+	void OnPostEngineInit();
+
 	TSharedRef<SDockTab> SpawnProductionWizard(const FSpawnTabArgs& SpawnTabArgs);
 	void RegisterMenus();
 
@@ -52,6 +53,9 @@ private:
 	void OpenAssemblyForEdit(UCineAssembly* Assembly);
 
 private:
+	/** Manages all integration of Cinematic Assembly Tools with Take Recorder. */
+	TUniquePtr<FCineAssemblyTakeRecorderIntegration> TakeRecorderIntegration;
+
 	FDelegateHandle OnSequenceSetHandle;
 };
 
@@ -85,14 +89,27 @@ void FCineAssemblyToolsEditorModule::StartupModule()
 
 	RegisterMenus();
 
-	FCoreDelegates::OnPostEngineInit.AddLambda([this]() 
-		{
-			TakeRecorderIntegration = MakeUnique<FCineAssemblyTakeRecorderIntegration>();
-		});
+	FCoreDelegates::OnPostEngineInit.AddRaw(this, &FCineAssemblyToolsEditorModule::OnPostEngineInit);
 
 	OnSequenceSetHandle = UMoviePipelineExecutorJob::OnSequenceSetEvent().AddRaw(this, &FCineAssemblyToolsEditorModule::OnSequenceSet);
 }
 
+void FCineAssemblyToolsEditorModule::OnPostEngineInit()
+{
+	TakeRecorderIntegration = MakeUnique<FCineAssemblyTakeRecorderIntegration>();
+
+	// Register with the asset tools module so that default asset names can be evaluated for tokens before being used to create new assets
+	IAssetTools& AssetTools = FModuleManager::GetModuleChecked<FAssetToolsModule>("AssetTools").Get();
+	AssetTools.RegisterSanitizeNameDelegate("CineAssemblyToolsEditorModule", FSanitizeName::CreateLambda([](FString& NameToSanitize)
+		{
+			if (UNamingTokensEngineSubsystem* NamingTokensSubsystem = GEngine->GetEngineSubsystem<UNamingTokensEngineSubsystem>())
+			{
+				FNamingTokenResultData Result = NamingTokensSubsystem->EvaluateTokenString(NameToSanitize);
+				NameToSanitize = Result.EvaluatedText.ToString();
+			}
+		}));
+}
+
 void FCineAssemblyToolsEditorModule::ShutdownModule()
 {
 	FCoreDelegates::OnPostEngineInit.RemoveAll(this);
@@ -104,6 +121,11 @@ void FCineAssemblyToolsEditorModule::ShutdownModule()
 	FGlobalTabmanager::Get()->UnregisterNomadTabSpawner(ProductionWizardTabName);
 
 	UMoviePipelineExecutorJob::OnSequenceSetEvent().Remove(OnSequenceSetHandle);
+
+	if (FAssetToolsModule* AssetToolsModule = FModuleManager::GetModulePtr<FAssetToolsModule>("AssetTools"))
+	{
+		AssetToolsModule->Get().UnregisterSanitizeNameDelegate("CineAssemblyToolsEditorModule");
+	}
 }
 
 TSharedRef<SDockTab> FCineAssemblyToolsEditorModule::SpawnProductionWizard(const FSpawnTabArgs& SpawnTabArgs)
diff --git a/Engine/Plugins/VirtualProduction/CinematicAssemblyTools/Source/CineAssemblyToolsEditor/Private/UI/SAssetNamingPanel.cpp b/Engine/Plugins/VirtualProduction/CinematicAssemblyTools/Source/CineAssemblyToolsEditor/Private/UI/SAssetNamingPanel.cpp
--- a/Engine/Plugins/VirtualProduction/CinematicAssemblyTools/Source/CineAssemblyToolsEditor/Private/UI/SAssetNamingPanel.cpp
+++ b/Engine/Plugins/VirtualProduction/CinematicAssemblyTools/Source/CineAssemblyToolsEditor/Private/UI/SAssetNamingPanel.cpp
@@ -370,7 +370,13 @@ bool SAssetNamingRow::ValidateDefaultAssetName(const FText& InText, FText& OutEr
 
 	// Ensure that the name does not contain any characters that would be invalid for an asset name
 	// This matches the validation that would happen if the user was renaming an asset in the content browser
-	return FName::IsValidXName(InText.ToString(), INVALID_OBJECTNAME_CHARACTERS INVALID_LONGPACKAGE_CHARACTERS, &OutErrorMessage);
+	FString InvalidCharacters = INVALID_OBJECTNAME_CHARACTERS INVALID_LONGPACKAGE_CHARACTERS;
+
+	// These characters are actually valid, because we want to support naming tokens
+	InvalidCharacters = InvalidCharacters.Replace(TEXT("{}"), TEXT(""));
+	InvalidCharacters = InvalidCharacters.Replace(TEXT(":"), TEXT(""));
+
+	return FName::IsValidXName(InText.ToString(), InvalidCharacters, &OutErrorMessage);
 }
 
 FReply SAssetNamingRow::DeleteRow()
diff --git a/Engine/Plugins/VirtualProduction/CinematicAssemblyTools/Source/CineAssemblyToolsEditor/Private/UI/SFolderHierarchyPanel.cpp b/Engine/Plugins/VirtualProduction/CinematicAssemblyTools/Source/CineAssemblyToolsEditor/Private/UI/SFolderHierarchyPanel.cpp
--- a/Engine/Plugins/VirtualProduction/CinematicAssemblyTools/Source/CineAssemblyToolsEditor/Private/UI/SFolderHierarchyPanel.cpp
+++ b/Engine/Plugins/VirtualProduction/CinematicAssemblyTools/Source/CineAssemblyToolsEditor/Private/UI/SFolderHierarchyPanel.cpp
@@ -9,11 +9,14 @@
 #include "ContentBrowserItem.h"
 #include "DirectoryPlaceholder.h"
 #include "DirectoryPlaceholderUtils.h"
+#include "Engine/Engine.h"
 #include "Framework/Application/SlateApplication.h"
 #include "Framework/MultiBox/MultiBoxBuilder.h"
 #include "IContentBrowserDataModule.h"
 #include "HAL/FileManager.h"
 #include "Misc/MessageDialog.h"
+#include "NamingTokensEngineSubsystem.h"
+#include "ObjectTools.h"
 #include "ProductionSettings.h"
 #include "Styling/StyleColors.h"
 #include "UI/SActiveProductionCombo.h"
@@ -440,8 +443,6 @@ void SFolderHierarchyPanel::BuildTreeFromProductionTemplate()
 	TOptional<const FCinematicProduction> ActiveProduction = ProductionSettings->GetActiveProduction();
 	if (ActiveProduction.IsSet())
 	{
-		FAssetRegistryModule& AssetRegistryModule = FModuleManager::GetModuleChecked<FAssetRegistryModule>("AssetRegistry");
-
 		// Sort the folder paths from the active production's template alphabetically to put them in the order they should appear in the tree
 		TArray<FFolderTemplate> ProductionTemplate = ActiveProduction.GetValue().TemplateFolders;
 		CachedInitialState = ProductionTemplate;
@@ -460,7 +461,7 @@ void SFolderHierarchyPanel::BuildTreeFromProductionTemplate()
 				TSharedPtr<FTemplateFolderTreeItem> NewItem = MakeShared<FTemplateFolderTreeItem>();
 				NewItem->Path.SetPathFromString(FolderTemplate.InternalPath, EContentBrowserPathType::Internal);
 
-				if (AssetRegistryModule.Get().PathExists(NewItem->Path.GetInternalPathString()))
+				if (DoesPathExist(NewItem->Path.GetInternalPathString()))
 				{
 					NewItem->Status = ETemplateFolderStatus::Exists;
 				}
@@ -484,8 +485,6 @@ void SFolderHierarchyPanel::BuildTreeFromProductionTemplate()
 
 FReply SFolderHierarchyPanel::OnAddFolderToTemplate()
 {
-	FAssetRegistryModule& AssetRegistryModule = FModuleManager::GetModuleChecked<FAssetRegistryModule>("AssetRegistry");
-
 	// Get the parent item for the new template folder being added (this can be the root folder if no parent is currently selected)
 	// The TreeView uses single selection mode, so at most one item can ever be selected by the user
 	TSharedPtr<FTemplateFolderTreeItem> SelectedTreeItem = RootItem;
@@ -502,7 +501,7 @@ FReply SFolderHierarchyPanel::OnAddFolderToTemplate()
 	const FString NewFolderName = CreateUniqueFolderName(SelectedTreeItem->Path);
 	NewFolder->Path.SetPathFromString(ParentPath / NewFolderName, EContentBrowserPathType::Internal);
 
-	if (AssetRegistryModule.Get().PathExists(NewFolder->Path.GetInternalPathString()))
+	if (DoesPathExist(NewFolder->Path.GetInternalPathString()))
 	{
 		NewFolder->Status = ETemplateFolderStatus::Exists;
 	}
@@ -574,8 +573,11 @@ bool SFolderHierarchyPanel::IsValidFolderName(const FText& InText, FText& OutErr
 		return false;
 	}
 
+	// These characters are actually valid, because we want to support naming tokens
+	const FString FolderNameWithoutTokenChars = InText.ToString().Replace(TEXT(":"), TEXT(""));
+
 	// Check for invalid characters
-	if (!AssetViewUtils::IsValidFolderName(InText.ToString(), OutErrorMessage))
+	if (!AssetViewUtils::IsValidFolderName(FolderNameWithoutTokenChars, OutErrorMessage))
 	{
 		return false;
 	}
@@ -656,7 +658,7 @@ void SFolderHierarchyPanel::SetTemplateFolderName(const FText& InText, ETextComm
 	SetChildrenPathRecursive(TreeItem, NewPath);
 
 	// Now that the path has changed, check again if this item exists, and update the status accordingly
-	if (AssetRegistryModule.Get().PathExists(NewPath))
+	if (DoesPathExist(NewPath))
 	{
 		TreeItem->Status = ETemplateFolderStatus::Exists;
 	}
@@ -730,7 +732,16 @@ void SFolderHierarchyPanel::CreateFolderFromTemplateRecursive(TSharedPtr<FTempla
 	if (TreeItem->Status == ETemplateFolderStatus::MissingCreate)
 	{
 		const FString PathToCreate = TreeItem->Path.GetInternalPathString();
-		const FString RelativeFilePath = FPackageName::LongPackageNameToFilename(PathToCreate);
+
+		// Evaluate any tokens found in the template path
+		UNamingTokensEngineSubsystem* NamingTokensSubsystem = GEngine->GetEngineSubsystem<UNamingTokensEngineSubsystem>();
+		FNamingTokenResultData Result = NamingTokensSubsystem->EvaluateTokenString(PathToCreate);	
+		const FString ResolvedPath = Result.EvaluatedText.ToString();
+
+		// Sanitize the name, in case any tokens were not able to be resolved
+		const FString SanitizedPath = ObjectTools::SanitizeInvalidChars(ResolvedPath, INVALID_LONGPACKAGE_CHARACTERS);
+	
+		const FString RelativeFilePath = FPackageName::LongPackageNameToFilename(SanitizedPath);
 		const FString AbsoluteFilePath = FPaths::ConvertRelativePathToFull(RelativeFilePath);
 
 		// Create the directory on disk, then add its path to the asset registry so it appears in Content Browser
@@ -740,7 +751,7 @@ void SFolderHierarchyPanel::CreateFolderFromTemplateRecursive(TSharedPtr<FTempla
 			if (IFileManager::Get().MakeDirectory(*AbsoluteFilePath, bCreateParentFoldersIfMissing))
 			{
 				FAssetRegistryModule& AssetRegistryModule = FModuleManager::GetModuleChecked<FAssetRegistryModule>("AssetRegistry");
-				AssetRegistryModule.Get().AddPath(PathToCreate);
+				AssetRegistryModule.Get().AddPath(SanitizedPath);
 			}
 		}
 
@@ -772,8 +783,6 @@ void SFolderHierarchyPanel::RemoveFolderFromTemplateRecursive(TSharedPtr<FTempla
 
 void SFolderHierarchyPanel::SetChildrenPathRecursive(TSharedPtr<FTemplateFolderTreeItem> TreeItem, const FString& NewPath)
 {
-	FAssetRegistryModule& AssetRegistryModule = FModuleManager::GetModuleChecked<FAssetRegistryModule>("AssetRegistry");
-
 	UProductionSettings* ProductionSettings = GetMutableDefault<UProductionSettings>();
 
 	for (TSharedPtr<FTemplateFolderTreeItem>& Child : TreeItem->Children)
@@ -786,7 +795,7 @@ void SFolderHierarchyPanel::SetChildrenPathRecursive(TSharedPtr<FTemplateFolderT
 		Child->Path.SetPathFromString(NewChildPath, EContentBrowserPathType::Internal);
 
 		// Now that the path has changed, check again if this item exists, and update the status accordingly
-		if (AssetRegistryModule.Get().PathExists(NewChildPath))
+		if (DoesPathExist(NewChildPath))
 		{
 			Child->Status = ETemplateFolderStatus::Exists;
 		}
@@ -890,6 +899,17 @@ bool SFolderHierarchyPanel::IsFolderEmpty(TSharedPtr<FTemplateFolderTreeItem> Tr
 	return bIsEmpty;
 }
 
+bool SFolderHierarchyPanel::DoesPathExist(const FString& Path) const 
+{
+	// Evaluate any tokens found in the template path before checking if the folder already exists
+	UNamingTokensEngineSubsystem* NamingTokensSubsystem = GEngine->GetEngineSubsystem<UNamingTokensEngineSubsystem>();
+	const FNamingTokenResultData Result = NamingTokensSubsystem->EvaluateTokenString(Path);
+	const FString ResolvedPath = Result.EvaluatedText.ToString();
+
+	FAssetRegistryModule& AssetRegistryModule = FModuleManager::GetModuleChecked<FAssetRegistryModule>("AssetRegistry");
+	return AssetRegistryModule.Get().PathExists(ResolvedPath);
+}
+
 const FSlateBrush* SFolderHierarchyPanel::GetFolderIconBadge(TSharedPtr<FTemplateFolderTreeItem> TreeItem) const
 {
 	if (TreeItem->Status == ETemplateFolderStatus::MissingCreate)
diff --git a/Engine/Plugins/VirtualProduction/CinematicAssemblyTools/Source/CineAssemblyToolsEditor/Private/UI/SFolderHierarchyPanel.h b/Engine/Plugins/VirtualProduction/CinematicAssemblyTools/Source/CineAssemblyToolsEditor/Private/UI/SFolderHierarchyPanel.h
--- a/Engine/Plugins/VirtualProduction/CinematicAssemblyTools/Source/CineAssemblyToolsEditor/Private/UI/SFolderHierarchyPanel.h
+++ b/Engine/Plugins/VirtualProduction/CinematicAssemblyTools/Source/CineAssemblyToolsEditor/Private/UI/SFolderHierarchyPanel.h
@@ -131,6 +131,9 @@ private:
 	/** Returns the status text to display next to the input tree item in its row in the tree view */
 	FText GetFolderStatusText(TSharedPtr<FTemplateFolderTreeItem> TreeItem) const;
 
+	/** Returns true if the path already exists in the asset registry */
+	bool DoesPathExist(const FString& Path) const;
+
 private:
 	/* Root path of the template folder tree */
 	static FName RootPath;
diff --git a/Engine/Source/Developer/AssetTools/Private/AssetTools.cpp b/Engine/Source/Developer/AssetTools/Private/AssetTools.cpp
--- a/Engine/Source/Developer/AssetTools/Private/AssetTools.cpp
+++ b/Engine/Source/Developer/AssetTools/Private/AssetTools.cpp
@@ -6578,11 +6578,15 @@ void UAssetToolsImpl::UnregisterSanitizeNameDelegate(const FName OwnerName)
 	SanitizeNameDelegates.Remove(OwnerName);
 }
 
-TOptional<FString> UAssetToolsImpl::GetDefaultAssetNameForClass(const UClass* Class) const
+TOptional<FString> UAssetToolsImpl::GetDefaultAssetNameForClass(const UClass* Class)
 {
-	if (const FString* AssetName = DefaultNewAssetNames.Find(Class))
+	if (FString* AssetName = DefaultNewAssetNames.Find(Class))
 	{
-		return *AssetName;
+		FString DefaultAssetName = *AssetName;
+		SanitizeName(DefaultAssetName);
+
+		FString SantizedDefaultName = ObjectTools::SanitizeInvalidChars(DefaultAssetName, INVALID_LONGPACKAGE_CHARACTERS INVALID_OBJECTNAME_CHARACTERS);
+		return SantizedDefaultName;
 	}
 
 	return TOptional<FString>();
@@ -6590,8 +6594,7 @@ TOptional<FString> UAssetToolsImpl::GetDefaultAssetNameForClass(const UClass* Cl
 
 void UAssetToolsImpl::RegisterDefaultAssetNameForClass(const UClass* Class, const FString& DefaultName)
 {
-	FString SantizedDefaultName = ObjectTools::SanitizeInvalidChars(DefaultName, INVALID_LONGPACKAGE_CHARACTERS INVALID_OBJECTNAME_CHARACTERS);
-	DefaultNewAssetNames.Add(Class, SantizedDefaultName);
+	DefaultNewAssetNames.Add(Class, DefaultName);
 }
 
 void UAssetToolsImpl::UnregisterDefaultAssetNameForClass(const UClass* Class)
diff --git a/Engine/Source/Developer/AssetTools/Private/AssetTools.h b/Engine/Source/Developer/AssetTools/Private/AssetTools.h
--- a/Engine/Source/Developer/AssetTools/Private/AssetTools.h
+++ b/Engine/Source/Developer/AssetTools/Private/AssetTools.h
@@ -160,7 +160,7 @@ public:
 	virtual bool SanitizeName(FString& NameToSanitize) override;
 	virtual void RegisterSanitizeNameDelegate(const FName OwnerName, FSanitizeName Delegate) override;
 	virtual void UnregisterSanitizeNameDelegate(const FName OwnerName) override;
-	virtual TOptional<FString> GetDefaultAssetNameForClass(const UClass* Class) const override;
+	virtual TOptional<FString> GetDefaultAssetNameForClass(const UClass* Class) override;
 	virtual void RegisterDefaultAssetNameForClass(const UClass* Class, const FString& DefaultName) override;
 	virtual void UnregisterDefaultAssetNameForClass(const UClass* Class) override;
 	virtual void RegisterCanMigrateAsset(const FName OwnerName, UE::AssetTools::FCanMigrateAsset Delegate) override;
diff --git a/Engine/Source/Developer/AssetTools/Public/IAssetTools.h b/Engine/Source/Developer/AssetTools/Public/IAssetTools.h
--- a/Engine/Source/Developer/AssetTools/Public/IAssetTools.h
+++ b/Engine/Source/Developer/AssetTools/Public/IAssetTools.h
@@ -702,7 +702,7 @@ public:
 	virtual void UnregisterSanitizeNameDelegate(const FName OwnerName) = 0;
 
 	/** Get the default name for the supplied class (if one is registered) to use when creating new assets of that type */
-	virtual TOptional<FString> GetDefaultAssetNameForClass(const UClass* Class) const = 0;
+	virtual TOptional<FString> GetDefaultAssetNameForClass(const UClass* Class) = 0;
 	/** Register a default name for the supplied class to use when creating new assets of that type */
 	virtual void RegisterDefaultAssetNameForClass(const UClass* Class, const FString& DefaultName) = 0;
 	/** Remove any previously registered name for the supplied class */
