ShotManagement: Add additional validation to Cine Assembly Schemas  to ensure that duplicate folders and duplicate asset names will not be allowed.
When adding a new folder to the hierarchy, its path will be unique.
Fixed a bug where the top level assembly item could be deleted from the tree if its parent folder was deleted.
Fixed a bug that did not allow for the naming token delimeter ":" to be used in folder names.
#rnx

diff --git a/Engine/Plugins/VirtualProduction/CinematicAssemblyTools/Source/CineAssemblyToolsEditor/Private/UI/CineAssembly/SCineAssemblySchemaWindow.cpp b/Engine/Plugins/VirtualProduction/CinematicAssemblyTools/Source/CineAssemblyToolsEditor/Private/UI/CineAssembly/SCineAssemblySchemaWindow.cpp
--- a/Engine/Plugins/VirtualProduction/CinematicAssemblyTools/Source/CineAssemblyToolsEditor/Private/UI/CineAssembly/SCineAssemblySchemaWindow.cpp
+++ b/Engine/Plugins/VirtualProduction/CinematicAssemblyTools/Source/CineAssemblyToolsEditor/Private/UI/CineAssembly/SCineAssemblySchemaWindow.cpp
@@ -2,6 +2,7 @@
 
 #include "SCineAssemblySchemaWindow.h"
 
+#include "Algo/Contains.h"
 #include "AssetRegistry/AssetRegistryModule.h"
 #include "AssetViewUtils.h"
 #include "CineAssemblyNamingTokens.h"
@@ -545,6 +546,40 @@ FReply SCineAssemblySchemaWindow::OnAddAssetClicked()
 	return FReply::Handled();
 }
 
+FString SCineAssemblySchemaWindow::MakeUniqueFolderPath(TSharedPtr<FSchemaTreeItem> InItem)
+{
+	// This implementation is based on a similar utility in AssetTools for creating a unique asset name.
+	const FString BaseName = InItem->Path / TEXT("NewFolder");
+
+	// Find the index in the string of the last non-numeric character
+	int32 CharIndex = BaseName.Len() - 1;
+	while (CharIndex >= 0 && BaseName[CharIndex] >= TEXT('0') && BaseName[CharIndex] <= TEXT('9'))
+	{
+		--CharIndex;
+	}
+
+	// Trim the numeric characters off the end of the BaseName string, but remember the integer that was trimmed off to increment and append to the output
+	int32 IntSuffix = 1;
+	FString TrimmedBaseName = BaseName;
+	if (CharIndex >= 0 && CharIndex < BaseName.Len() - 1)
+	{
+		TrimmedBaseName = BaseName.Left(CharIndex + 1);
+
+		const FString TrailingInteger = BaseName.RightChop(CharIndex + 1);
+		IntSuffix = FCString::Atoi(*TrailingInteger) + 1;
+	}
+
+	FString WorkingName = TrimmedBaseName;
+
+	while (Algo::ContainsBy(InItem->ChildFolders, WorkingName, &FSchemaTreeItem::Path))
+	{
+		WorkingName = FString::Printf(TEXT("%s%d"), *TrimmedBaseName, IntSuffix);
+		IntSuffix++;
+	}
+
+	return WorkingName;
+}
+
 FReply SCineAssemblySchemaWindow::OnAddFolderClicked()
 {
 	TSharedPtr<FSchemaTreeItem> NewFolder = MakeShared<FSchemaTreeItem>();
@@ -561,8 +596,7 @@ FReply SCineAssemblySchemaWindow::OnAddFolderClicked()
 	}
 
 	const FString ParentPath = SelectedTreeItem->Path;
-	const FString NewFolderName = TEXT("NewFolder");
-	NewFolder->Path = ParentPath / NewFolderName;
+	NewFolder->Path = MakeUniqueFolderPath(SelectedTreeItem);
 	NewFolder->Parent = SelectedTreeItem;
 
 	SelectedTreeItem->ChildFolders.Add(NewFolder);
@@ -1048,8 +1082,11 @@ bool SCineAssemblySchemaWindow::ValidateFolderName(const FText& InText, FText& O
 		return false;
 	}
 
+	// These characters are actually valid, because we want to support naming tokens
+	const FString FolderNameWithoutTokenChars = InText.ToString().Replace(TEXT(":"), TEXT(""));
+
 	// Check for invalid characters
-	if (!AssetViewUtils::IsValidFolderName(InText.ToString(), OutErrorMessage))
+	if (!AssetViewUtils::IsValidFolderName(FolderNameWithoutTokenChars, OutErrorMessage))
 	{
 		return false;
 	}
@@ -1081,6 +1118,11 @@ void SCineAssemblySchemaWindow::DeleteTreeItem(TSharedPtr<FSchemaTreeItem> TreeI
 {
 	if (TreeItem->Type == FSchemaTreeItem::EItemType::Folder)
 	{
+		if (ContainsRecursive(TreeItem, TopLevelAssemblyItem))
+		{
+			TreeItem->Parent->ChildAssets.Add(TopLevelAssemblyItem);
+		}
+
 		TreeItem->Parent->ChildFolders.Remove(TreeItem);
 	}
 	else
@@ -1090,6 +1132,22 @@ void SCineAssemblySchemaWindow::DeleteTreeItem(TSharedPtr<FSchemaTreeItem> TreeI
 	TreeView->RequestTreeRefresh();
 }
 
+bool SCineAssemblySchemaWindow::ContainsRecursive(TSharedPtr<FSchemaTreeItem> InParent, TSharedPtr<FSchemaTreeItem> InItem)
+{
+	if (InParent->ChildAssets.Contains(InItem))
+	{
+		return true;
+	}
+
+	bool bFoundRecursively = false;
+	for (TSharedPtr<FSchemaTreeItem> ChildFolder : InParent->ChildFolders)
+	{
+		bFoundRecursively |= ContainsRecursive(ChildFolder, InItem);
+	}
+
+	return bFoundRecursively;
+}
+
 void SCineAssemblySchemaWindow::OnSchemaPropertiesChanged(const FPropertyChangedEvent& PropertyChangedEvent)
 {
 	const FName PropertyName = PropertyChangedEvent.GetPropertyName();
@@ -1257,6 +1315,22 @@ FReply SCineAssemblySchemaWindow::OnTreeRowAcceptDrop(const FDragDropEvent& InDr
 			NewItem->Type = FSchemaTreeItem::EItemType::Asset;
 		}
 
+		// If an item with the same name already exists in the drop location, do not handle the drop (the dragged item will be reset to its original location)
+		if (NewItem->Type == FSchemaTreeItem::EItemType::Asset)
+		{
+			if (Algo::ContainsBy(InItem->ChildAssets, InItem->Path / Operation->ItemName, &FSchemaTreeItem::Path))
+			{
+				return FReply::Unhandled();
+			}
+		}
+		else if (NewItem->Type == FSchemaTreeItem::EItemType::Folder)
+		{
+			if (Algo::ContainsBy(InItem->ChildFolders, InItem->Path / Operation->ItemName, &FSchemaTreeItem::Path))
+			{
+				return FReply::Unhandled();
+			}
+		}
+
 		const FString ParentPath = InItem->Path;
 		NewItem->Path = ParentPath / Operation->ItemName;
 
diff --git a/Engine/Plugins/VirtualProduction/CinematicAssemblyTools/Source/CineAssemblyToolsEditor/Private/UI/CineAssembly/SCineAssemblySchemaWindow.h b/Engine/Plugins/VirtualProduction/CinematicAssemblyTools/Source/CineAssemblyToolsEditor/Private/UI/CineAssembly/SCineAssemblySchemaWindow.h
--- a/Engine/Plugins/VirtualProduction/CinematicAssemblyTools/Source/CineAssemblyToolsEditor/Private/UI/CineAssembly/SCineAssemblySchemaWindow.h
+++ b/Engine/Plugins/VirtualProduction/CinematicAssemblyTools/Source/CineAssemblyToolsEditor/Private/UI/CineAssembly/SCineAssemblySchemaWindow.h
@@ -215,6 +215,12 @@ private:
 	/** Callback when an item is dropped onto a folder in the content tree view to add it to the children of that row */
 	FReply OnTreeRowAcceptDrop(const FDragDropEvent& InDragDropEvent, EItemDropZone InDropZone, TSharedPtr<FSchemaTreeItem> InItem);
 
+	/** Generates a unique folder path name, assuming the input item will be the parent */
+	FString MakeUniqueFolderPath(TSharedPtr<FSchemaTreeItem> InItem);
+
+	/** Returns true if the input parent, or any of its child folders, contains in the input item */
+	bool ContainsRecursive(TSharedPtr<FSchemaTreeItem> InParent, TSharedPtr<FSchemaTreeItem> InItem);
+
 private:
 	/** Switcher that controls which menu tab is currently visible */
 	TSharedPtr<SWidgetSwitcher> MenuTabSwitcher;
