ShotManagement: Fix class list in default asset naming tab of the Production Wizard to only show classes that have a Factory associated with them. This more closely matches the list of assets in the New Asset Context Menu which is what opens when right-clicking in Content Browser.
#rnx

diff --git a/Engine/Plugins/VirtualProduction/CinematicAssemblyTools/Source/CineAssemblyToolsEditor/Private/UI/SAssetNamingPanel.cpp b/Engine/Plugins/VirtualProduction/CinematicAssemblyTools/Source/CineAssemblyToolsEditor/Private/UI/SAssetNamingPanel.cpp
--- a/Engine/Plugins/VirtualProduction/CinematicAssemblyTools/Source/CineAssemblyToolsEditor/Private/UI/SAssetNamingPanel.cpp
+++ b/Engine/Plugins/VirtualProduction/CinematicAssemblyTools/Source/CineAssemblyToolsEditor/Private/UI/SAssetNamingPanel.cpp
@@ -29,10 +29,11 @@ namespace UE::CineAssemblyTools::AssetNaming::Private
 	public:
 		FAssetNamingClassFilter()
 		{
-			TArray<TSoftClassPtr<UObject>> RegisteredAssetClasses = UAssetDefinitionRegistry::Get()->GetAllRegisteredAssetClasses();
-			for (const TSoftClassPtr<UObject>& AsssetClass : RegisteredAssetClasses)
+			IAssetTools& AssetTools = FModuleManager::GetModuleChecked<FAssetToolsModule>("AssetTools").Get();
+			TArray<UFactory*> Factories = AssetTools.GetNewAssetFactories();
+			for (UFactory* Factory : Factories)
 			{
-				AllowedAssetClasses.Add(AsssetClass.Get());
+				AllowedAssetClasses.Add(Factory->GetSupportedClass());
 			}
 		}
 
@@ -271,6 +272,7 @@ TSharedRef<SWidget> SAssetNamingRow::GenerateWidgetForColumn(const FName& Column
 							.MetaClass(UObject::StaticClass())
 							.SelectedClass_Lambda([this]() { return AssetNaming->Class; })
 							.ClassViewerFilters({ MakeShared<FAssetNamingClassFilter>() })
+							.ShowDisplayNames(true)
 							.OnSetClass(this, &SAssetNamingRow::SetAssetClass)
 					]
 			];
