ShotManagement: Fix Assembly and Schema windows to store the transient assets being configured as strong object pointers to prevent them from being GC'd before the user is finished.
Rename the "No Schema" asset used in the Assembly creation window to "NoSchema", to fix a warning that an asset name cannot have a space in it.
Fix Schema picker in the Assembly creation window to not reset the schema when the user clicks in the empty space of the picker widget.
#rnx

diff --git a/Engine/Plugins/VirtualProduction/CinematicAssemblyTools/Source/CineAssemblyToolsEditor/Private/UI/CineAssembly/SCineAssemblyConfigWindow.cpp b/Engine/Plugins/VirtualProduction/CinematicAssemblyTools/Source/CineAssemblyToolsEditor/Private/UI/CineAssembly/SCineAssemblyConfigWindow.cpp
--- a/Engine/Plugins/VirtualProduction/CinematicAssemblyTools/Source/CineAssemblyToolsEditor/Private/UI/CineAssembly/SCineAssemblyConfigWindow.cpp
+++ b/Engine/Plugins/VirtualProduction/CinematicAssemblyTools/Source/CineAssemblyToolsEditor/Private/UI/CineAssembly/SCineAssemblyConfigWindow.cpp
@@ -47,7 +47,7 @@ void SCineAssemblyConfigWindow::Construct(const FArguments& InArgs, const FStrin
 	CreateAssetPath = InCreateAssetPath;
 
 	// Create a new transient CineAssembly to configure in the UI.
-	CineAssemblyToConfigure = NewObject<UCineAssembly>(GetTransientPackage(), NAME_None, RF_Transient);
+	CineAssemblyToConfigure = TStrongObjectPtr<UCineAssembly>(NewObject<UCineAssembly>(GetTransientPackage(), NAME_None, RF_Transient));
 
 	UWorld* CurrentWorld = GEditor->GetEditorWorldContext().World();
 	CineAssemblyToConfigure->Level = FSoftObjectPath(CurrentWorld);
@@ -124,7 +124,7 @@ TSharedRef<SWidget> SCineAssemblyConfigWindow::MakeCineTemplatePanel()
 	Config.OnAssetSelected = FOnAssetSelected::CreateSP(this, &SCineAssemblyConfigWindow::OnSchemaSelected);
 
 	const FString PackageName = IPluginManager::Get().FindPlugin(UE_PLUGIN_NAME)->GetMountedAssetPath();
-	const FString NoSchemaName = TEXT("No Schema");
+	const FString NoSchemaName = TEXT("NoSchema");
 	const FString PackagePath = PackageName / NoSchemaName;
 
 	FAssetData NoSchemaAssetData = FAssetData(*PackagePath, *PackageName, *NoSchemaName, FTopLevelAssetPath());
@@ -158,13 +158,18 @@ TSharedRef<SWidget> SCineAssemblyConfigWindow::MakeCineTemplatePanel()
 
 void SCineAssemblyConfigWindow::OnSchemaSelected(const FAssetData& AssetData)
 {
-	if (UCineAssemblySchema* CineAssemblySchema = Cast<UCineAssemblySchema>(AssetData.GetAsset()))
+	const FString NoSchemaName = TEXT("NoSchema");
+	if (AssetData.AssetName == NoSchemaName)
+	{
+		SelectedSchema = nullptr;
+	}
+	else if (UCineAssemblySchema* CineAssemblySchema = Cast<UCineAssemblySchema>(AssetData.GetAsset()))
 	{
 		SelectedSchema = CineAssemblySchema;
 	}
 	else
 	{
-		SelectedSchema = nullptr;
+		return;
 	}
 
 	CineAssemblyToConfigure->ChangeSchema(SelectedSchema);
@@ -178,7 +183,7 @@ void SCineAssemblyConfigWindow::OnSchemaSelected(const FAssetData& AssetData)
 
 TSharedRef<SWidget> SCineAssemblyConfigWindow::MakeInfoPanel()
 {
-	DetailsView = UE::CineAssemblyWidgets::Private::MakeDetailsWidget(CineAssemblyToConfigure);
+	DetailsView = UE::CineAssemblyWidgets::Private::MakeDetailsWidget(CineAssemblyToConfigure.Get());
 
 	TabSwitcher = SNew(SWidgetSwitcher)
 		+ SWidgetSwitcher::Slot()
@@ -275,7 +280,7 @@ void SCineAssemblyConfigWindow::EvaluateTokenString(FTemplateString& StringToEva
 	FDateTime CurrentTime = FDateTime::Now();
 	if ((CurrentTime - LastTokenUpdateTime).GetSeconds() >= 1.0f)
 	{
-		StringToEvaluate.Resolved = UCineAssemblyNamingTokens::GetResolvedText(StringToEvaluate.Template, CineAssemblyToConfigure);
+		StringToEvaluate.Resolved = UCineAssemblyNamingTokens::GetResolvedText(StringToEvaluate.Template, CineAssemblyToConfigure.Get());
 		LastTokenUpdateTime = CurrentTime;
 	}
 }
@@ -461,7 +466,7 @@ void SCineAssemblyConfigWindow::PopulateHierarchyTree()
 void SCineAssemblyConfigWindow::EvaluateHierarchyTokensRecursive(TSharedPtr<FHierarchyTreeItem> TreeItem)
 {
 	// Evaluate the token template string for this tree item
-	TreeItem->Path.Resolved = UCineAssemblyNamingTokens::GetResolvedText(TreeItem->Path.Template, CineAssemblyToConfigure);
+	TreeItem->Path.Resolved = UCineAssemblyNamingTokens::GetResolvedText(TreeItem->Path.Template, CineAssemblyToConfigure.Get());
 
 	// Evaluate the tokens for all of the child assets, then resort them alphabetically based on the resolved paths
 	for (const TSharedPtr<FHierarchyTreeItem>& Asset : TreeItem->ChildAssets)
@@ -630,7 +635,7 @@ TSharedRef<SWidget> SCineAssemblyConfigWindow::MakeButtonsPanel()
 
 FReply SCineAssemblyConfigWindow::OnCreateAssetClicked()
 {
-	UCineAssemblyFactory::CreateConfiguredAssembly(CineAssemblyToConfigure, CreateAssetPath);
+	UCineAssemblyFactory::CreateConfiguredAssembly(CineAssemblyToConfigure.Get(), CreateAssetPath);
 
 	RequestDestroyWindow();
 	return FReply::Handled();
diff --git a/Engine/Plugins/VirtualProduction/CinematicAssemblyTools/Source/CineAssemblyToolsEditor/Private/UI/CineAssembly/SCineAssemblyConfigWindow.h b/Engine/Plugins/VirtualProduction/CinematicAssemblyTools/Source/CineAssemblyToolsEditor/Private/UI/CineAssembly/SCineAssemblyConfigWindow.h
--- a/Engine/Plugins/VirtualProduction/CinematicAssemblyTools/Source/CineAssemblyToolsEditor/Private/UI/CineAssembly/SCineAssemblyConfigWindow.h
+++ b/Engine/Plugins/VirtualProduction/CinematicAssemblyTools/Source/CineAssemblyToolsEditor/Private/UI/CineAssembly/SCineAssemblyConfigWindow.h
@@ -107,7 +107,7 @@ private:
 	TSharedPtr<IDetailsView> DetailsView;
 
 	/** Transient object used only by this UI to configure the properties of the new asset that will get created by the Factory */
-	UCineAssembly* CineAssemblyToConfigure = nullptr;
+	TStrongObjectPtr<UCineAssembly> CineAssemblyToConfigure = nullptr;
 
 	/** The currently selected schema to use as a base for configuring the new cine assembly */
 	UCineAssemblySchema* SelectedSchema = nullptr;
diff --git a/Engine/Plugins/VirtualProduction/CinematicAssemblyTools/Source/CineAssemblyToolsEditor/Private/UI/CineAssembly/SCineAssemblySchemaWindow.cpp b/Engine/Plugins/VirtualProduction/CinematicAssemblyTools/Source/CineAssemblyToolsEditor/Private/UI/CineAssembly/SCineAssemblySchemaWindow.cpp
--- a/Engine/Plugins/VirtualProduction/CinematicAssemblyTools/Source/CineAssemblyToolsEditor/Private/UI/CineAssembly/SCineAssemblySchemaWindow.cpp
+++ b/Engine/Plugins/VirtualProduction/CinematicAssemblyTools/Source/CineAssemblyToolsEditor/Private/UI/CineAssembly/SCineAssemblySchemaWindow.cpp
@@ -36,15 +36,15 @@ void SCineAssemblySchemaWindow::Construct(const FArguments& InArgs, const FStrin
 {
 	// Create a new transient CineAssemblySchema to configure in the UI.
 	// If the configuration is successful, this will turn into the persistent object created by the factory.
-	SchemaToConfigure = NewObject<UCineAssemblySchema>(GetTransientPackage(), NAME_None, RF_Transient);
+	UCineAssemblySchema* NewSchema = NewObject<UCineAssemblySchema>(GetTransientPackage(), NAME_None, RF_Transient);
 	CreateAssetPath = InCreateAssetPath;
 
-	Construct(InArgs, SchemaToConfigure, ESchemaConfigMode::CreateNew);
+	Construct(InArgs, NewSchema, ESchemaConfigMode::CreateNew);
 }
 
 void SCineAssemblySchemaWindow::Construct(const FArguments& InArgs, UCineAssemblySchema* InSchema, ESchemaConfigMode InConfigMode)
 {
-	SchemaToConfigure = InSchema;
+	SchemaToConfigure = TStrongObjectPtr<UCineAssemblySchema>(InSchema);
 	Mode = InConfigMode;
 
 	const FVector2D DefaultWindowSize = FVector2D(1400, 750);
@@ -366,7 +366,7 @@ TSharedRef<SWidget> SCineAssemblySchemaWindow::MakeDetailsWidget(bool bShowMetad
 
 	DetailsView->SetIsPropertyVisibleDelegate(FIsPropertyVisible::CreateSP(this, &SCineAssemblySchemaWindow::IsPropertyVisible, bShowMetadata));
 
-	DetailsView->SetObject(SchemaToConfigure, true);
+	DetailsView->SetObject(SchemaToConfigure.Get(), true);
 	DetailsView->OnFinishedChangingProperties().AddSP(this, &SCineAssemblySchemaWindow::OnSchemaPropertiesChanged);
 
 	return DetailsView;
@@ -1163,7 +1163,7 @@ void SCineAssemblySchemaWindow::OnSchemaPropertiesChanged(const FPropertyChanged
 
 FReply SCineAssemblySchemaWindow::OnCreateAssetClicked()
 {
-	UCineAssemblySchemaFactory::CreateConfiguredSchema(SchemaToConfigure, CreateAssetPath);
+	UCineAssemblySchemaFactory::CreateConfiguredSchema(SchemaToConfigure.Get(), CreateAssetPath);
 
 	RequestDestroyWindow();
 	return FReply::Handled();
diff --git a/Engine/Plugins/VirtualProduction/CinematicAssemblyTools/Source/CineAssemblyToolsEditor/Private/UI/CineAssembly/SCineAssemblySchemaWindow.h b/Engine/Plugins/VirtualProduction/CinematicAssemblyTools/Source/CineAssemblyToolsEditor/Private/UI/CineAssembly/SCineAssemblySchemaWindow.h
--- a/Engine/Plugins/VirtualProduction/CinematicAssemblyTools/Source/CineAssemblyToolsEditor/Private/UI/CineAssembly/SCineAssemblySchemaWindow.h
+++ b/Engine/Plugins/VirtualProduction/CinematicAssemblyTools/Source/CineAssemblyToolsEditor/Private/UI/CineAssembly/SCineAssemblySchemaWindow.h
@@ -229,7 +229,7 @@ private:
 	TSharedPtr<SWidgetSwitcher> AssetTabSwitcher;
 
 	/** Transient object used only by this UI to configure the properties of the new asset that will get created by the factory */
-	UCineAssemblySchema* SchemaToConfigure = nullptr;
+	TStrongObjectPtr<UCineAssemblySchema> SchemaToConfigure = nullptr;
 
 	/** Mode that configures the UI based on the intended user interactions with the Cine Assembly schema asset */
 	ESchemaConfigMode Mode = ESchemaConfigMode::CreateNew;
