ShotManagement: Fix default schema assets not being visible in the assembly and schema windows by default by introducing a config property to save the user preference specifically for asset pickers in those tools.
#rnx

diff --git a/Engine/Plugins/VirtualProduction/CinematicAssemblyTools/Source/CineAssemblyToolsEditor/Private/UI/CineAssembly/SCineAssemblyConfigWindow.cpp b/Engine/Plugins/VirtualProduction/CinematicAssemblyTools/Source/CineAssemblyToolsEditor/Private/UI/CineAssembly/SCineAssemblyConfigWindow.cpp
--- a/Engine/Plugins/VirtualProduction/CinematicAssemblyTools/Source/CineAssemblyToolsEditor/Private/UI/CineAssembly/SCineAssemblyConfigWindow.cpp
+++ b/Engine/Plugins/VirtualProduction/CinematicAssemblyTools/Source/CineAssemblyToolsEditor/Private/UI/CineAssembly/SCineAssemblyConfigWindow.cpp
@@ -10,9 +10,11 @@
 #include "IContentBrowserSingleton.h"
 #include "Interfaces/IPluginManager.h"
 #include "IStructureDetailsView.h"
+#include "Misc/ConfigCacheIni.h"
 #include "Modules/ModuleManager.h"
 #include "PropertyEditorModule.h"
 #include "UObject/Package.h"
+#include "Settings/ContentBrowserSettings.h"
 #include "STemplateStringEditableTextBox.h"
 #include "Widgets/SBoxPanel.h"
 #include "Widgets/Input/SButton.h"
@@ -52,6 +54,8 @@ void SCineAssemblyConfigWindow::Construct(const FArguments& InArgs, const FStrin
 	UWorld* CurrentWorld = GEditor->GetEditorWorldContext().World();
 	CineAssemblyToConfigure->Level = FSoftObjectPath(CurrentWorld);
 
+	GetOnWindowClosedEvent().AddSP(this, &SCineAssemblyConfigWindow::OnAssemblyWindowClosed);
+
 	const FVector2D DefaultWindowSize = FVector2D(1400, 750);
 
 	SWindow::Construct(SWindow::FArguments()
@@ -102,6 +106,22 @@ void SCineAssemblyConfigWindow::Construct(const FArguments& InArgs, const FStrin
 		]);
 }
 
+void SCineAssemblyConfigWindow::OnAssemblyWindowClosed(const TSharedRef<SWindow>& InWindow)
+{
+	// Save the UI config settings for whether to display engine/plugin content
+	UContentBrowserSettings* ContentBrowserSettings = GetMutableDefault<UContentBrowserSettings>();
+
+	const bool bShowEngineContent = ContentBrowserSettings->GetDisplayEngineFolder();
+	const bool bShowPluginContent = ContentBrowserSettings->GetDisplayPluginFolders();
+
+	GConfig->SetBool(TEXT("NewCineAssemblyUI"), TEXT("bShowEngineContent"), bShowEngineContent, GEditorPerProjectIni);
+	GConfig->SetBool(TEXT("NewCineAssemblyUI"), TEXT("bShowPluginContent"), bShowPluginContent, GEditorPerProjectIni);
+	GConfig->Flush(false, GEditorPerProjectIni);
+
+	ContentBrowserSettings->SetDisplayEngineFolder(bShowEngineContentCached);
+	ContentBrowserSettings->SetDisplayPluginFolders(bShowPluginContentCached);
+}
+
 TSharedRef<SWidget> SCineAssemblyConfigWindow::MakeCineTemplatePanel()
 {
 	FContentBrowserModule& ContentBrowserModule = FModuleManager::Get().LoadModuleChecked<FContentBrowserModule>("ContentBrowser");
@@ -148,6 +168,20 @@ TSharedRef<SWidget> SCineAssemblyConfigWindow::MakeCineTemplatePanel()
 			return SNew(SToolTip).Text(LOCTEXT("NoSchemaToolTip", "Create a new assembly with no schema"));
 		});
 
+	// Check the UI config settings to determine whether or not to display engine/plugin content by default in this window
+	UContentBrowserSettings* ContentBrowserSettings = GetMutableDefault<UContentBrowserSettings>();
+
+	bool bShowEngineContent = true;
+	bool bShowPluginContent = true;
+	GConfig->GetBool(TEXT("NewCineAssemblyUI"), TEXT("bShowEngineContent"), bShowEngineContent, GEditorPerProjectIni);
+	GConfig->GetBool(TEXT("NewCineAssemblyUI"), TEXT("bShowPluginContent"), bShowPluginContent, GEditorPerProjectIni);
+
+	bShowEngineContentCached = ContentBrowserSettings->GetDisplayEngineFolder();
+	bShowPluginContentCached = ContentBrowserSettings->GetDisplayPluginFolders();
+
+	ContentBrowserSettings->SetDisplayEngineFolder(bShowEngineContent);
+	ContentBrowserSettings->SetDisplayPluginFolders(bShowPluginContent);
+
 	return SNew(SBorder)
 		.BorderImage(FCineAssemblyToolsStyle::Get().GetBrush("ProductionWizard.PanelBackground"))
 		.Padding(16.0f)
diff --git a/Engine/Plugins/VirtualProduction/CinematicAssemblyTools/Source/CineAssemblyToolsEditor/Private/UI/CineAssembly/SCineAssemblyConfigWindow.h b/Engine/Plugins/VirtualProduction/CinematicAssemblyTools/Source/CineAssemblyToolsEditor/Private/UI/CineAssembly/SCineAssemblyConfigWindow.h
--- a/Engine/Plugins/VirtualProduction/CinematicAssemblyTools/Source/CineAssemblyToolsEditor/Private/UI/CineAssembly/SCineAssemblyConfigWindow.h
+++ b/Engine/Plugins/VirtualProduction/CinematicAssemblyTools/Source/CineAssemblyToolsEditor/Private/UI/CineAssembly/SCineAssemblyConfigWindow.h
@@ -72,6 +72,9 @@ private:
 	/** Closes the window and indicates that no assets should be created by the asset factory */
 	FReply OnCancelClicked();
 
+	/** Callback when this window is about to be closed */
+	void OnAssemblyWindowClosed(const TSharedRef<SWindow>& InWindow);
+
 	/** Updates the UI and CineAssembly properties based on the selected schema */
 	void OnSchemaSelected(const FAssetData& AssetData);
 
@@ -126,6 +129,10 @@ private:
 
 	/** A read-only tree view of folders and assets that will be created based on the selected schema */
 	TSharedPtr<STreeView<TSharedPtr<FHierarchyTreeItem>>> HierarchyTreeView;
+
+	/** Cached content browser settings, used to restore defaults when closing the window */
+	bool bShowEngineContentCached = false;
+	bool bShowPluginContentCached = false;
 };
 
 /** A panel that displays properties of a Cine Assembly asset */
diff --git a/Engine/Plugins/VirtualProduction/CinematicAssemblyTools/Source/CineAssemblyToolsEditor/Private/UI/CineAssembly/SCineAssemblySchemaWindow.cpp b/Engine/Plugins/VirtualProduction/CinematicAssemblyTools/Source/CineAssemblyToolsEditor/Private/UI/CineAssembly/SCineAssemblySchemaWindow.cpp
--- a/Engine/Plugins/VirtualProduction/CinematicAssemblyTools/Source/CineAssemblyToolsEditor/Private/UI/CineAssembly/SCineAssemblySchemaWindow.cpp
+++ b/Engine/Plugins/VirtualProduction/CinematicAssemblyTools/Source/CineAssemblyToolsEditor/Private/UI/CineAssembly/SCineAssemblySchemaWindow.cpp
@@ -10,10 +10,12 @@
 #include "CineAssemblyToolsStyle.h"
 #include "Framework/MultiBox/MultiBoxBuilder.h"
 #include "IStructureDetailsView.h"
+#include "Misc/ConfigCacheIni.h"
 #include "Modules/ModuleManager.h"
 #include "NamingTokensEngineSubsystem.h"
 #include "PropertyEditorDelegates.h"
 #include "PropertyEditorModule.h"
+#include "Settings/ContentBrowserSettings.h"
 #include "Styling/StyleColors.h"
 #include "UObject/Package.h"
 #include "Widgets/Input/SButton.h"
@@ -62,6 +64,20 @@ void SCineAssemblySchemaWindow::Construct(const FArguments& InArgs, UCineAssembl
 	// Initialize the content tree view with the current list of assets and folders saved in the schema
 	InitializeContentTree();
 
+	// Check the UI config settings to determine whether or not to display engine/plugin content by default in this window
+	UContentBrowserSettings* ContentBrowserSettings = GetMutableDefault<UContentBrowserSettings>();
+
+	bool bShowEngineContent = true;
+	bool bShowPluginContent = true;
+	GConfig->GetBool(TEXT("NewCineAssemblySchemaUI"), TEXT("bShowEngineContent"), bShowEngineContent, GEditorPerProjectIni);
+	GConfig->GetBool(TEXT("NewCineAssemblySchemaUI"), TEXT("bShowPluginContent"), bShowPluginContent, GEditorPerProjectIni);
+
+	bShowEngineContentCached = ContentBrowserSettings->GetDisplayEngineFolder();
+	bShowPluginContentCached = ContentBrowserSettings->GetDisplayPluginFolders();
+
+	ContentBrowserSettings->SetDisplayEngineFolder(bShowEngineContent);
+	ContentBrowserSettings->SetDisplayPluginFolders(bShowPluginContent);
+
 	GetOnWindowClosedEvent().AddSP(this, &SCineAssemblySchemaWindow::OnSchemaWindowClosed);
 
 	SWindow::Construct(SWindow::FArguments()
@@ -1217,6 +1233,19 @@ void SCineAssemblySchemaWindow::OnSchemaWindowClosed(const TSharedRef<SWindow>&
 		SchemaToConfigure->Modify();
 		SchemaToConfigure->SubsequencesToCreate = AssetList;
 	}
+
+	// Save the UI config settings for whether to display engine/plugin content
+	UContentBrowserSettings* ContentBrowserSettings = GetMutableDefault<UContentBrowserSettings>();
+
+	const bool bShowEngineContent = ContentBrowserSettings->GetDisplayEngineFolder();
+	const bool bShowPluginContent = ContentBrowserSettings->GetDisplayPluginFolders();
+
+	GConfig->SetBool(TEXT("NewCineAssemblySchemaUI"), TEXT("bShowEngineContent"), bShowEngineContent, GEditorPerProjectIni);
+	GConfig->SetBool(TEXT("NewCineAssemblySchemaUI"), TEXT("bShowPluginContent"), bShowPluginContent, GEditorPerProjectIni);
+	GConfig->Flush(false, GEditorPerProjectIni);
+
+	ContentBrowserSettings->SetDisplayEngineFolder(bShowEngineContentCached);
+	ContentBrowserSettings->SetDisplayPluginFolders(bShowPluginContentCached);
 }
 
 FReply SCineAssemblySchemaWindow::OnAssetRowDragDetected(const FGeometry& MyGeometry, const FPointerEvent& MouseEvent)
diff --git a/Engine/Plugins/VirtualProduction/CinematicAssemblyTools/Source/CineAssemblyToolsEditor/Private/UI/CineAssembly/SCineAssemblySchemaWindow.h b/Engine/Plugins/VirtualProduction/CinematicAssemblyTools/Source/CineAssemblyToolsEditor/Private/UI/CineAssembly/SCineAssemblySchemaWindow.h
--- a/Engine/Plugins/VirtualProduction/CinematicAssemblyTools/Source/CineAssemblyToolsEditor/Private/UI/CineAssembly/SCineAssemblySchemaWindow.h
+++ b/Engine/Plugins/VirtualProduction/CinematicAssemblyTools/Source/CineAssemblyToolsEditor/Private/UI/CineAssembly/SCineAssemblySchemaWindow.h
@@ -221,6 +221,10 @@ private:
 	/** Returns true if the input parent, or any of its child folders, contains in the input item */
 	bool ContainsRecursive(TSharedPtr<FSchemaTreeItem> InParent, TSharedPtr<FSchemaTreeItem> InItem);
 
+	/** Cached content browser settings, used to restore defaults when closing the window */
+	bool bShowEngineContentCached = false;
+	bool bShowPluginContentCached = false;
+
 private:
 	/** Switcher that controls which menu tab is currently visible */
 	TSharedPtr<SWidgetSwitcher> MenuTabSwitcher;
