ShotManagement: Fix Cine Assembly edit windows to be dockable as nomad tabs.
The CAT module manages a map of Tab IDs to Assembly IDs, and handles serializing the list of currently open tabs out to a config file so that they can be restored in the layout the next time the editor is initialized. The current maximum number of Assembly Tabs that can be open at one time is 16.
The CineAssembly asset now features a serialized GUID property to help uniquely identify assets between runs of UE and after asset name changes.
It is very possible that the layout is restored before the asset registry is finished scanning, so if there are any Assembly Tabs that were opened, the underlying widget handles updating its UI once the asset registry is ready.
#rnx

diff --git a/Engine/Plugins/VirtualProduction/CinematicAssemblyTools/Source/CineAssemblyTools/Private/CineAssembly.cpp b/Engine/Plugins/VirtualProduction/CinematicAssemblyTools/Source/CineAssemblyTools/Private/CineAssembly.cpp
--- a/Engine/Plugins/VirtualProduction/CinematicAssemblyTools/Source/CineAssemblyTools/Private/CineAssembly.cpp
+++ b/Engine/Plugins/VirtualProduction/CinematicAssemblyTools/Source/CineAssemblyTools/Private/CineAssembly.cpp
@@ -21,10 +21,16 @@
 #define LOCTEXT_NAMESPACE "CineAssembly"
 
 const FName UCineAssembly::AssetRegistryTag_AssemblyType = "AssemblyType";
+const FName UCineAssembly::AssemblyGuidPropertyName = GET_MEMBER_NAME_CHECKED(UCineAssembly, AssemblyGuid);
 
 UCineAssembly::UCineAssembly()
 {
 	MetadataJsonObject = MakeShared<FJsonObject>();
+
+	if (!AssemblyGuid.IsValid())
+	{
+		AssemblyGuid = FGuid::NewGuid();
+	}
 }
 
 void UCineAssembly::Initialize()
@@ -32,6 +38,11 @@ void UCineAssembly::Initialize()
 	Super::Initialize();
 }
 
+FGuid UCineAssembly::GetAssemblyGuid() const
+{
+	return AssemblyGuid;
+}
+
 const UCineAssemblySchema* UCineAssembly::GetSchema() const
 {
 	return BaseSchema.Get();
diff --git a/Engine/Plugins/VirtualProduction/CinematicAssemblyTools/Source/CineAssemblyTools/Public/CineAssembly.h b/Engine/Plugins/VirtualProduction/CinematicAssemblyTools/Source/CineAssemblyTools/Public/CineAssembly.h
--- a/Engine/Plugins/VirtualProduction/CinematicAssemblyTools/Source/CineAssemblyTools/Public/CineAssembly.h
+++ b/Engine/Plugins/VirtualProduction/CinematicAssemblyTools/Source/CineAssemblyTools/Public/CineAssembly.h
@@ -40,6 +40,9 @@ public:
 	virtual void Serialize(FArchive& Ar) override;
 	//~ End UObject Interface
 
+	/** Get the unique ID of this assembly */
+	CINEASSEMBLYTOOLS_API FGuid GetAssemblyGuid() const;
+
 	/** Get the base schema for this assembly */
 	CINEASSEMBLYTOOLS_API const UCineAssemblySchema* GetSchema() const;
 
@@ -148,8 +151,13 @@ public:
 
 	/** The asset registry tag that contains the assembly type information */
 	static CINEASSEMBLYTOOLS_API const FName AssetRegistryTag_AssemblyType;
+	static CINEASSEMBLYTOOLS_API const FName AssemblyGuidPropertyName;
 
 private:
+	/** Unique ID for this assembly, assigned at object creation */
+	UPROPERTY(AssetRegistrySearchable, meta = (IgnoreForMemberInitializationTest))
+	FGuid AssemblyGuid;
+
 	/** The schema that was used as a base when creating this assembly (can be null if no schema was used) */
 	UPROPERTY()
 	TObjectPtr<UCineAssemblySchema> BaseSchema;
diff --git a/Engine/Plugins/VirtualProduction/CinematicAssemblyTools/Source/CineAssemblyToolsEditor/Private/CineAssemblyToolsEditorModule.cpp b/Engine/Plugins/VirtualProduction/CinematicAssemblyTools/Source/CineAssemblyToolsEditor/Private/CineAssemblyToolsEditorModule.cpp
--- a/Engine/Plugins/VirtualProduction/CinematicAssemblyTools/Source/CineAssemblyToolsEditor/Private/CineAssemblyToolsEditorModule.cpp
+++ b/Engine/Plugins/VirtualProduction/CinematicAssemblyTools/Source/CineAssemblyToolsEditor/Private/CineAssemblyToolsEditorModule.cpp
@@ -2,6 +2,7 @@
 
 #include "Modules/ModuleManager.h"
 
+#include "AssetRegistry/AssetRegistryModule.h"
 #include "AssetToolsModule.h"
 #include "CineAssembly.h"
 #include "CineAssemblySchema.h"
@@ -13,6 +14,7 @@
 #include "Framework/Application/SlateApplication.h"
 #include "Framework/Docking/TabManager.h"
 #include "Interfaces/IMainFrameModule.h"
+#include "Misc/ConfigCacheIni.h"
 #include "MoviePipelineQueue.h"
 #include "NamingTokensEngineSubsystem.h"
 #include "ProductionSettings.h"
@@ -29,6 +31,8 @@
 
 #define LOCTEXT_NAMESPACE "CineAssemblyToolsEditorModule"
 
+DEFINE_LOG_CATEGORY_STATIC(LogCineAssemblyToolsEditorModule, Log, All)
+
 static const FName ProductionWizardTabName("ProductionWizard");
 
 class FCineAssemblyToolsEditorModule : public IModuleInterface
@@ -43,7 +47,15 @@ public:
 private:
 	void OnPostEngineInit();
 
+	/** Spawns a nomad tab for the Production Wizard tool */
 	TSharedRef<SDockTab> SpawnProductionWizard(const FSpawnTabArgs& SpawnTabArgs);
+
+	/** Spawns a nomad tab for an Assembly asset to edit its properties */
+	TSharedRef<SDockTab> SpawnAssemblyTab(const FSpawnTabArgs& SpawnTabArgs, UCineAssembly* Assembly);
+	TSharedRef<SDockTab> SpawnAssemblyTab(const FSpawnTabArgs& SpawnTabArgs, FGuid AssemblyGuid);
+	TSharedRef<SDockTab> SpawnAssemblyTab(const FSpawnTabArgs& SpawnTabArgs, TSharedPtr<SCineAssemblyEditWidget> Widget);
+
+	/** Add context menu actions */
 	void RegisterMenus();
 
 	/** Called whenever a job has its sequence set. */
@@ -52,13 +64,27 @@ private:
 	/** Opens a window to edit the properties of the input assembly */
 	void OpenAssemblyForEdit(UCineAssembly* Assembly);
 
+	/** Returns the next available tab name in the map of managed tabs. If there are no available tabs, returns false */
+	bool TryGetNextAssemblyTab(FTabId& OutTabId);
+
+	/** Callback when an asset is deleted which allows us to check if it was a Cine Assembly and then close the tab for that Assembly (if one was open) */
+	void OnAssetDeleted(UObject* Object);
+
 private:
 	/** Manages all integration of Cinematic Assembly Tools with Take Recorder. */
 	TUniquePtr<FCineAssemblyTakeRecorderIntegration> TakeRecorderIntegration;
 
+	/** Map of TabIds to AssemblyIDs to track which assets are open in which tabs */
+	TMap<FTabId, FGuid> ManagedTabs;
+	const int32 MaxAssemblyTabs = 16;
+
+	static const FString OpenAssemblySection;
+
 	FDelegateHandle OnSequenceSetHandle;
 };
 
+const FString FCineAssemblyToolsEditorModule::OpenAssemblySection = TEXT("OpenAssemblyTabs");
+
 void FCineAssemblyToolsEditorModule::StartupModule()
 {
 	FCineAssemblyToolsStyle::Get();
@@ -92,6 +118,13 @@ void FCineAssemblyToolsEditorModule::StartupModule()
 	FCoreDelegates::OnPostEngineInit.AddRaw(this, &FCineAssemblyToolsEditorModule::OnPostEngineInit);
 
 	OnSequenceSetHandle = UMoviePipelineExecutorJob::OnSequenceSetEvent().AddRaw(this, &FCineAssemblyToolsEditorModule::OnSequenceSet);
+
+	// Initialize the map of managed tabs
+	for (int32 TabIndex = 0; TabIndex < MaxAssemblyTabs; ++TabIndex)
+	{
+		const FString AssemblyTabName = FString::Printf(TEXT("CineAssemblyTab%d"), TabIndex);
+		ManagedTabs.Add(FTabId(*AssemblyTabName), FGuid());
+	}
 }
 
 void FCineAssemblyToolsEditorModule::OnPostEngineInit()
@@ -108,6 +141,25 @@ void FCineAssemblyToolsEditorModule::OnPostEngineInit()
 				NameToSanitize = Result.EvaluatedText.ToString();
 			}
 		}));
+
+	// Load any previously opened assembly tabs and register new nomad tab spawners for them so they can be properly restored in the layout.
+	for (TPair<FTabId, FGuid>& TabPair : ManagedTabs)
+	{
+		const FName TabName = TabPair.Key.TabType;
+
+		FString GuidString;
+		if (GConfig->GetString(*OpenAssemblySection, *TabName.ToString(), GuidString, GEditorPerProjectIni))
+		{
+			const FGuid AssemblyID = FGuid(GuidString);
+			TabPair.Value = AssemblyID;
+
+			FGlobalTabmanager::Get()->RegisterNomadTabSpawner(TabName, FOnSpawnTab::CreateRaw(this, &FCineAssemblyToolsEditorModule::SpawnAssemblyTab, AssemblyID))
+				.SetMenuType(ETabSpawnerMenuType::Hidden);
+		}
+	}
+
+	const FAssetRegistryModule& AssetRegistryModule = FModuleManager::GetModuleChecked<FAssetRegistryModule>("AssetRegistry");
+	AssetRegistryModule.Get().OnInMemoryAssetDeleted().AddRaw(this, &FCineAssemblyToolsEditorModule::OnAssetDeleted);
 }
 
 void FCineAssemblyToolsEditorModule::ShutdownModule()
@@ -126,6 +178,28 @@ void FCineAssemblyToolsEditorModule::ShutdownModule()
 	{
 		AssetToolsModule->Get().UnregisterSanitizeNameDelegate("CineAssemblyToolsEditorModule");
 	}
+
+	if (FAssetRegistryModule* AssetRegistryModule = FModuleManager::GetModulePtr<FAssetRegistryModule>("AssetRegistry"))
+	{
+		AssetRegistryModule->Get().OnInMemoryAssetDeleted().RemoveAll(this);
+	}
+
+	// Write out each managed tab that is associated with a valid Assembly to a config file so that we can restore them the next time the editor launches
+	GConfig->EmptySection(*OpenAssemblySection, GEditorPerProjectIni);
+	for (const TPair<FTabId, FGuid>& TabPair : ManagedTabs)
+	{
+		const FGuid AssemblyID = TabPair.Value;
+		if (AssemblyID.IsValid())
+		{
+			GConfig->AddToSection(*OpenAssemblySection, TabPair.Key.TabType, AssemblyID.ToString(), GEditorPerProjectIni);
+
+			if (TSharedPtr<FGlobalTabmanager> GlobalTabManager = FGlobalTabmanager::Get())
+			{
+				GlobalTabManager->UnregisterNomadTabSpawner(TabPair.Key.TabType);
+			}
+		}
+	}
+	GConfig->Flush(false, GEditorPerProjectIni);
 }
 
 TSharedRef<SDockTab> FCineAssemblyToolsEditorModule::SpawnProductionWizard(const FSpawnTabArgs& SpawnTabArgs)
@@ -137,6 +211,41 @@ TSharedRef<SDockTab> FCineAssemblyToolsEditorModule::SpawnProductionWizard(const
 		];
 }
 
+TSharedRef<SDockTab> FCineAssemblyToolsEditorModule::SpawnAssemblyTab(const FSpawnTabArgs& SpawnTabArgs, FGuid AssemblyGuid)
+{
+	TSharedPtr<SCineAssemblyEditWidget> Widget = SNew(SCineAssemblyEditWidget, AssemblyGuid);
+	return SpawnAssemblyTab(SpawnTabArgs, Widget);
+}
+
+TSharedRef<SDockTab> FCineAssemblyToolsEditorModule::SpawnAssemblyTab(const FSpawnTabArgs& SpawnTabArgs, UCineAssembly* Assembly)
+{
+	TSharedPtr<SCineAssemblyEditWidget> Widget = SNew(SCineAssemblyEditWidget, Assembly);
+	return SpawnAssemblyTab(SpawnTabArgs, Widget);
+}
+
+TSharedRef<SDockTab> FCineAssemblyToolsEditorModule::SpawnAssemblyTab(const FSpawnTabArgs& SpawnTabArgs, TSharedPtr<SCineAssemblyEditWidget> Widget)
+{
+	TSharedRef<SDockTab> AssemblyTab = SNew(SDockTab)
+		.TabRole(ETabRole::NomadTab)
+		.Label_Lambda([Widget]() -> FText
+			{
+				return FText::FromString(Widget->GetAssemblyName());
+			})
+		[
+			Widget.ToSharedRef()
+		];
+
+	// Unregister the spawner and reset the map entry when the tab closes
+	const FTabId TabId = SpawnTabArgs.GetTabId();
+	AssemblyTab->SetOnTabClosed(SDockTab::FOnTabClosedCallback::CreateLambda([this, TabId](TSharedRef<SDockTab> Tab)
+		{
+			FGlobalTabmanager::Get()->UnregisterNomadTabSpawner(*TabId.ToString());
+			ManagedTabs[TabId].Invalidate();
+		}));
+
+	return AssemblyTab;
+}
+
 void FCineAssemblyToolsEditorModule::RegisterMenus()
 {
 	UToolMenus* ToolMenus = UToolMenus::Get();
@@ -187,18 +296,66 @@ void FCineAssemblyToolsEditorModule::RegisterMenus()
 
 void FCineAssemblyToolsEditorModule::OpenAssemblyForEdit(UCineAssembly* Assembly)
 {
-	TSharedRef<SCineAssemblyEditWidget> CineAssemblyWidget = SNew(SCineAssemblyEditWidget, Assembly);
-
-	const FVector2D DefaultWindowSize = FVector2D(600, 800);
-	TSharedRef<SWindow> EditWindow = SNew(SWindow)
-		.Title(LOCTEXT("EditCineAssemblyTitle", "Edit Cine Assembly"))
-		.SizingRule(ESizingRule::UserSized)
-		.ClientSize(DefaultWindowSize)
-		.SupportsMaximize(false)
-		.SupportsMinimize(false);
-
-	EditWindow->SetContent(CineAssemblyWidget);
-	FSlateApplication::Get().AddWindow(EditWindow);
+	// Check if the assembly is already open in one of the managed tabs, and if so, bring focus to the existing tab
+	const FGuid AssemblyID = Assembly->GetAssemblyGuid();
+	const FTabId* ExistingTab = ManagedTabs.FindKey(AssemblyID);
+
+	if (ExistingTab)
+	{
+		FGlobalTabmanager::Get()->TryInvokeTab(*ExistingTab);
+		return;
+	}
+
+	// Try to get the TabID for the next available tab in the map.
+	FTabId NewTabId;
+	if (!TryGetNextAssemblyTab(NewTabId))
+	{
+		return;
+	}
+
+	// Register a new spawner and invoke a new tab to edit the properties of the assembly asset
+	FTabSpawnerEntry& Tmp = FGlobalTabmanager::Get()->RegisterNomadTabSpawner(NewTabId.TabType, FOnSpawnTab::CreateRaw(this, &FCineAssemblyToolsEditorModule::SpawnAssemblyTab, Assembly))
+		.SetMenuType(ETabSpawnerMenuType::Hidden);
+
+	FGlobalTabmanager::Get()->TryInvokeTab(NewTabId);
+
+	ManagedTabs[NewTabId] = AssemblyID;
+}
+
+bool FCineAssemblyToolsEditorModule::TryGetNextAssemblyTab(FTabId& OutTabId)
+{
+	// Find the first tab in the map that is not currently mapped to a valid assembly
+	for (const TPair<FTabId, FGuid>& TabPair : ManagedTabs)
+	{
+		if (!TabPair.Value.IsValid())
+		{
+			OutTabId = TabPair.Key;
+			return true;
+		}
+	}
+
+	// If all tabs are currently mapped, we return false so that no new tab is opened
+	UE_LOG(LogCineAssemblyToolsEditorModule, Warning, TEXT("The Cine Assembly tab could not be opened because the maximum number of Assembly tabs are already open."));
+	return false;
+}
+
+void FCineAssemblyToolsEditorModule::OnAssetDeleted(UObject* Object)
+{
+	if (const UCineAssembly* Assembly = Cast<UCineAssembly>(Object))
+	{
+		// Find the entry in the tab map matching the ID of the Assembly being deleted (if there is one)
+		if (const FTabId* OpenTabId = ManagedTabs.FindKey(Assembly->GetAssemblyGuid()))
+		{
+			// Get the open tab that is already open for this assembly and then close it. 
+			constexpr bool bInvokeAsInactive = true;
+			if (TSharedPtr<SDockTab> AssemblyTab = FGlobalTabmanager::Get()->TryInvokeTab(*OpenTabId, bInvokeAsInactive))
+			{
+				AssemblyTab->RequestCloseTab();
+			}
+
+			return;
+		}
+	}
 }
 
 void FCineAssemblyToolsEditorModule::OnSequenceSet(UMoviePipelineExecutorJob* Job, ULevelSequence* Sequence)
diff --git a/Engine/Plugins/VirtualProduction/CinematicAssemblyTools/Source/CineAssemblyToolsEditor/Private/UI/CineAssembly/SCineAssemblyConfigWindow.cpp b/Engine/Plugins/VirtualProduction/CinematicAssemblyTools/Source/CineAssemblyToolsEditor/Private/UI/CineAssembly/SCineAssemblyConfigWindow.cpp
--- a/Engine/Plugins/VirtualProduction/CinematicAssemblyTools/Source/CineAssemblyToolsEditor/Private/UI/CineAssembly/SCineAssemblyConfigWindow.cpp
+++ b/Engine/Plugins/VirtualProduction/CinematicAssemblyTools/Source/CineAssemblyToolsEditor/Private/UI/CineAssembly/SCineAssemblyConfigWindow.cpp
@@ -2,6 +2,7 @@
 
 #include "SCineAssemblyConfigWindow.h"
 
+#include "AssetRegistry/AssetRegistryModule.h"
 #include "CineAssemblyFactory.h"
 #include "CineAssemblyNamingTokens.h"
 #include "CineAssemblyToolsStyle.h"
@@ -684,54 +685,124 @@ FReply SCineAssemblyConfigWindow::OnCancelClicked()
 void SCineAssemblyEditWidget::Construct(const FArguments& InArgs, UCineAssembly* InAssembly)
 {
 	CineAssembly = InAssembly;
+	ChildSlot [ BuildUI() ];
+}
+
+void SCineAssemblyEditWidget::Construct(const FArguments& InArgs, FGuid InAssemblyGuid)
+{
+	// The UI will be temporary because no CineAssembly has been found yet
+	ChildSlot [ BuildUI() ];
+
+	// If the asset registry is still scanning assets, add a callback to find the assembly asset matching the input GUID and update this widget once the scan is finished.
+	// Otherwise, we can find the assembly asset and update the UI immediately. 
+	const FAssetRegistryModule& AssetRegistryModule = FModuleManager::GetModuleChecked<FAssetRegistryModule>("AssetRegistry");
+	if (AssetRegistryModule.Get().IsLoadingAssets())
+	{
+		AssetRegistryModule.Get().OnFilesLoaded().AddSP(this, &SCineAssemblyEditWidget::FindAssembly, InAssemblyGuid);
+	}
+	else
+	{
+		FindAssembly(InAssemblyGuid);
+	}
+}
+
+SCineAssemblyEditWidget::~SCineAssemblyEditWidget()
+{
+	if (FAssetRegistryModule* AssetRegistryModule = FModuleManager::GetModulePtr<FAssetRegistryModule>("AssetRegistry"))
+	{
+		AssetRegistryModule->Get().OnFilesLoaded().RemoveAll(this);
+	}
+}
+
+TSharedRef<SWidget> SCineAssemblyEditWidget::BuildUI()
+{
+	// Build a temporary UI to display while waiting for the assembly to be loaded
+	if (!CineAssembly)
+	{
+		return SNew(SBorder)
+			.BorderImage(FCineAssemblyToolsStyle::Get().GetBrush("Borders.PanelNoBorder"))
+			.Padding(8.0f)
+			[
+				SNew(STextBlock).Text(LOCTEXT("LoadingAssemblyText", "Loading Cine Assembly..."))
+			];
+	}
 
 	TabSwitcher = SNew(SWidgetSwitcher)
 		+ SWidgetSwitcher::Slot()
 		[MakeOverviewWidget()]
 
 		+ SWidgetSwitcher::Slot()
-		[ UE::CineAssemblyWidgets::Private::MakeDetailsWidget(InAssembly) ];
+		[UE::CineAssemblyWidgets::Private::MakeDetailsWidget(CineAssembly)];
 
-	ChildSlot
+	return SNew(SBorder)
+		.BorderImage(FCineAssemblyToolsStyle::Get().GetBrush("Borders.PanelNoBorder"))
+		.Padding(8.0f)
 		[
-			SNew(SBorder)
-				.BorderImage(FCineAssemblyToolsStyle::Get().GetBrush("Borders.PanelNoBorder"))
-				.Padding(8.0f)
-				[
-					SNew(SScrollBox)
-					.Orientation(Orient_Vertical)
+			SNew(SScrollBox)
+				.Orientation(Orient_Vertical)
 
-					+SScrollBox::Slot()
-						.AutoSize()
-						.HAlign(HAlign_Center)
-						.Padding(0.0f, 0.0f, 0.0f, 8.0f)
-						[
-							SNew(SSegmentedControl<int32>)
-								.Value(0)
-								.OnValueChanged_Lambda([Switcher = TabSwitcher](int32 NewValue)
-									{
-										Switcher->SetActiveWidgetIndex(NewValue);
-									})
+			+ SScrollBox::Slot()
+				.AutoSize()
+				.HAlign(HAlign_Center)
+				.Padding(0.0f, 0.0f, 0.0f, 8.0f)
+				[
+					SNew(SSegmentedControl<int32>)
+						.Value(0)
+						.OnValueChanged_Lambda([Switcher = TabSwitcher](int32 NewValue)
+							{
+								Switcher->SetActiveWidgetIndex(NewValue);
+							})
 
-								+ SSegmentedControl<int32>::Slot(0)
-								.Text(LOCTEXT("OverviewTab", "Overview"))
-								.Icon(FSlateIcon(FCineAssemblyToolsStyle::StyleName, "Icons.Animation").GetIcon())
+						+ SSegmentedControl<int32>::Slot(0)
+						.Text(LOCTEXT("OverviewTab", "Overview"))
+						.Icon(FSlateIcon(FCineAssemblyToolsStyle::StyleName, "Icons.Animation").GetIcon())
 
-								+ SSegmentedControl<int32>::Slot(1)
-								.Text(LOCTEXT("DetailsTab", "Details"))
-								.Icon(FSlateIcon(FAppStyle::GetAppStyleSetName(), "Icons.Details").GetIcon())
-						]
+						+ SSegmentedControl<int32>::Slot(1)
+						.Text(LOCTEXT("DetailsTab", "Details"))
+						.Icon(FSlateIcon(FAppStyle::GetAppStyleSetName(), "Icons.Details").GetIcon())
+				]
 
-					+ SScrollBox::Slot()
-						.FillSize(1.0f)
-						.HAlign(HAlign_Fill)
-						[
-							TabSwitcher.ToSharedRef()
-						]
+			+ SScrollBox::Slot()
+				.FillSize(1.0f)
+				.HAlign(HAlign_Fill)
+				[
+					TabSwitcher.ToSharedRef()
 				]
 		];
 }
 
+void SCineAssemblyEditWidget::FindAssembly(FGuid AssemblyID)
+{
+	const FAssetRegistryModule& AssetRegistryModule = FModuleManager::GetModuleChecked<FAssetRegistryModule>("AssetRegistry");
+
+	// The only search criterion for the asset search is for an asset with an AssemblyID matching the input GUID
+	const TMultiMap<FName, FString> TagValues = { { UCineAssembly::AssemblyGuidPropertyName, AssemblyID.ToString() } };
+
+	TArray<FAssetData> AssemblyAssets;
+	AssetRegistryModule.Get().GetAssetsByTagValues(TagValues, AssemblyAssets);
+
+	// The Assembly ID is unique, so at most one asset should ever be found
+	if (AssemblyAssets.Num() > 0)
+	{
+		CineAssembly = Cast<UCineAssembly>(AssemblyAssets[0].GetAsset());
+
+		// Update the widget's UI
+		ChildSlot.DetachWidget();
+		ChildSlot.AttachWidget(BuildUI());
+	}
+}
+
+FString SCineAssemblyEditWidget::GetAssemblyName()
+{
+	if (CineAssembly)
+	{
+		FString AssemblyName;
+		CineAssembly->GetName(AssemblyName);
+		return AssemblyName;
+	}
+	return TEXT("CineAssembly");
+}
+
 TSharedRef<SWidget> SCineAssemblyEditWidget::MakeOverviewWidget()
 {
 	const UCineAssemblySchema* Schema = CineAssembly->GetSchema();
diff --git a/Engine/Plugins/VirtualProduction/CinematicAssemblyTools/Source/CineAssemblyToolsEditor/Private/UI/CineAssembly/SCineAssemblyConfigWindow.h b/Engine/Plugins/VirtualProduction/CinematicAssemblyTools/Source/CineAssemblyToolsEditor/Private/UI/CineAssembly/SCineAssemblyConfigWindow.h
--- a/Engine/Plugins/VirtualProduction/CinematicAssemblyTools/Source/CineAssemblyToolsEditor/Private/UI/CineAssembly/SCineAssemblyConfigWindow.h
+++ b/Engine/Plugins/VirtualProduction/CinematicAssemblyTools/Source/CineAssemblyToolsEditor/Private/UI/CineAssembly/SCineAssemblyConfigWindow.h
@@ -140,13 +140,31 @@ class SCineAssemblyEditWidget : public SCompoundWidget
 {
 public:
 	SCineAssemblyEditWidget() = default;
+	~SCineAssemblyEditWidget();
 
 	SLATE_BEGIN_ARGS(SCineAssemblyEditWidget) {}
 	SLATE_END_ARGS()
 
+	/** Widget construction, initialized with the assembly asset being edited */
 	void Construct(const FArguments& InArgs, UCineAssembly* InAssembly);
 
+	/**
+	 * Widget construction, initialized with the GUID of the assembly to be edited
+	 * The widget will search the asset registry to find the assembly asset with the matching GUID,
+	 * and then update the widget contents accordingly. 
+	 */
+	void Construct(const FArguments& InArgs, FGuid InAssemblyGuid);
+
+	/** Returns the name of the assembly asset being edited */
+	FString GetAssemblyName();
+
+	/** Searches the asset registry for a Cine Assembly matching the input ID and updates the UI */
+	void FindAssembly(FGuid AssemblyID);
+
 private:
+	/** Constructs the main UI for the widget */ 
+	TSharedRef<SWidget> BuildUI();
+
 	/** Creates the widget to display for the Overview tab */
 	TSharedRef<SWidget> MakeOverviewWidget();
 
