ShotManagement: Add {parent} and {production} default tokens.
Modify the metadata token function to check if the metadata string value could refer to an asset, and returns only the name of the asset instead of the full path..
#rnx

diff --git a/Engine/Plugins/VirtualProduction/CinematicAssemblyTools/Source/CineAssemblyTools/Private/CineAssemblyNamingTokens.cpp b/Engine/Plugins/VirtualProduction/CinematicAssemblyTools/Source/CineAssemblyTools/Private/CineAssemblyNamingTokens.cpp
--- a/Engine/Plugins/VirtualProduction/CinematicAssemblyTools/Source/CineAssemblyTools/Private/CineAssemblyNamingTokens.cpp
+++ b/Engine/Plugins/VirtualProduction/CinematicAssemblyTools/Source/CineAssemblyTools/Private/CineAssemblyNamingTokens.cpp
@@ -63,6 +63,24 @@ void UCineAssemblyNamingTokens::OnCreateDefaultTokens(TArray<FNamingTokenData>&
 			return FText::GetEmpty();
 		};
 
+	auto ParentTokenFunc = [](TWeakObjectPtr<UCineAssembly> InAssembly) -> FText
+		{
+			if (InAssembly.IsValid() && InAssembly->ParentAssembly.IsValid())
+			{
+				return FText::FromString(InAssembly->ParentAssembly.GetAssetName());
+			}
+			return FText::GetEmpty();
+		};
+
+	auto ProductionTokenFunc = [](TWeakObjectPtr<UCineAssembly> InAssembly) -> FText
+		{
+			if (InAssembly.IsValid() && InAssembly->Production.IsValid())
+			{
+				return FText::FromString(InAssembly->ProductionName);
+			}
+			return FText::GetEmpty();
+		};
+
 	FNamingTokenData AssemblyNameToken;
 	AssemblyNameToken.TokenKey = TEXT("assembly");
 	AssemblyNameToken.DisplayName = LOCTEXT("AssemblyNameToken", "Assembly Name");
@@ -78,9 +96,21 @@ void UCineAssemblyNamingTokens::OnCreateDefaultTokens(TArray<FNamingTokenData>&
 	LevelToken.DisplayName = LOCTEXT("TargetLevelTokenName", "Target Level");
 	LevelToken.TokenProcessorNative.BindLambda([this, LevelTokenFunc]() { return ExecuteTokenFunc(LevelTokenFunc); });
 
+	FNamingTokenData ParentToken;
+	ParentToken.TokenKey = TEXT("parent");
+	ParentToken.DisplayName = LOCTEXT("ParentTokenName", "Parent Assembly");
+	ParentToken.TokenProcessorNative.BindLambda([this, ParentTokenFunc]() { return ExecuteTokenFunc(ParentTokenFunc); });
+
+	FNamingTokenData ProductionToken;
+	ProductionToken.TokenKey = TEXT("production");
+	ProductionToken.DisplayName = LOCTEXT("ProductionTokenName", "Production");
+	ProductionToken.TokenProcessorNative.BindLambda([this, ProductionTokenFunc]() { return ExecuteTokenFunc(ProductionTokenFunc); });
+
 	Tokens.Add(AssemblyNameToken);
 	Tokens.Add(SchemaToken);
 	Tokens.Add(LevelToken);
+	Tokens.Add(ParentToken);
+	Tokens.Add(ProductionToken);
 }
 
 void UCineAssemblyNamingTokens::OnPreEvaluate_Implementation(const FNamingTokensEvaluationData& InEvaluationData)
@@ -114,6 +144,13 @@ void UCineAssemblyNamingTokens::AddMetadataToken(const FString& InTokenKey)
 				FString ValueString;
 				if (InAssembly->GetMetadataAsString(InTokenKey, ValueString))
 				{
+					// Test if the value string could be an object path, and if it is, return just the filename instead of the full path
+					FSoftObjectPath ObjectPath = FSoftObjectPath(ValueString);
+					if (ObjectPath.IsValid())
+					{
+						ValueString = FPaths::GetBaseFilename(ValueString);
+					}
+
 					return FText::FromString(ValueString);
 				}
 			}
diff --git a/Engine/Plugins/VirtualProduction/CinematicAssemblyTools/Source/CineAssemblyTools/Public/CineAssembly.h b/Engine/Plugins/VirtualProduction/CinematicAssemblyTools/Source/CineAssemblyTools/Public/CineAssembly.h
--- a/Engine/Plugins/VirtualProduction/CinematicAssemblyTools/Source/CineAssemblyTools/Public/CineAssembly.h
+++ b/Engine/Plugins/VirtualProduction/CinematicAssemblyTools/Source/CineAssemblyTools/Public/CineAssembly.h
@@ -102,6 +102,9 @@ public:
 	UFUNCTION(BlueprintPure, Category = "Cine Assembly Tools", meta = (ReturnDisplayName = "FoundKey"))
 	CINEASSEMBLYTOOLS_API bool GetMetadataAsNumber(FString InKey, double& OutValue) const;
 
+	/** Adds a new metadata key to the list of supported naming tokens for assemblies */
+	CINEASSEMBLYTOOLS_API void AddMetadataNamingToken(const FString& InKey);
+
 private:
 	/** Sets the base schema for this assembly and re-initializes the metadata inherited from the schema */
 	CINEASSEMBLYTOOLS_API void ChangeSchema(UCineAssemblySchema* InSchema);
@@ -109,9 +112,6 @@ private:
 	/** Update the underlying json object whenever keys/values in the instance metadata map are added/removed/modified */
 	void UpdateInstanceMetadata();
 
-	/** Adds a new metadata key to the list of supported naming tokens for assemblies */
-	void AddMetadataNamingToken(const FString& InKey);
-
 public:
 	/** The assembly name, which supports tokens */
 	UPROPERTY()
@@ -137,6 +137,10 @@ public:
 	UPROPERTY(EditAnywhere, Category = "Default")
 	FGuid Production;
 
+	/** The name of the Cinematic Production that this assembly is associated with */
+	UPROPERTY()
+	FString ProductionName;
+
 	/** Array of template names (possibly containing tokens), based on the schema, used to create the SubAssemblies */
 	UPROPERTY()
 	TArray<FTemplateString> SubAssemblyNames;
diff --git a/Engine/Plugins/VirtualProduction/CinematicAssemblyTools/Source/CineAssemblyToolsEditor/Private/CineAssemblyCustomization.cpp b/Engine/Plugins/VirtualProduction/CinematicAssemblyTools/Source/CineAssemblyToolsEditor/Private/CineAssemblyCustomization.cpp
--- a/Engine/Plugins/VirtualProduction/CinematicAssemblyTools/Source/CineAssemblyToolsEditor/Private/CineAssemblyCustomization.cpp
+++ b/Engine/Plugins/VirtualProduction/CinematicAssemblyTools/Source/CineAssemblyToolsEditor/Private/CineAssemblyCustomization.cpp
@@ -129,6 +129,8 @@ void FCineAssemblyCustomization::CustomizeMetadataCategory(IDetailLayoutBuilder&
 			continue;
 		}
 
+		CustomizedCineAssembly->AddMetadataNamingToken(MetadataDesc.Key);
+
 		TSharedPtr<SWidget> ValueWidget;
 		if (MetadataDesc.Type == ECineAssemblyMetadataType::String)
 		{
@@ -411,6 +413,7 @@ TSharedRef<SWidget> FCineAssemblyCustomization::BuildProductionNameMenu()
 			{ 
 				CustomizedCineAssembly->Modify();
 				CustomizedCineAssembly->Production = FGuid();
+				CustomizedCineAssembly->ProductionName = TEXT("None");
 			})),
 		NAME_None,
 		EUserInterfaceActionType::None
@@ -430,6 +433,7 @@ TSharedRef<SWidget> FCineAssemblyCustomization::BuildProductionNameMenu()
 				{
 					CustomizedCineAssembly->Modify();
 					CustomizedCineAssembly->Production = Production.ProductionID;
+					CustomizedCineAssembly->ProductionName = Production.ProductionName;
 				})),
 			NAME_None,
 			EUserInterfaceActionType::None
