ShotManagement: Add Preview image to cine assembly assets based on the asset thumbnail.
#rnx

diff --git a/Engine/Plugins/VirtualProduction/CinematicAssemblyTools/Source/CineAssemblyToolsEditor/Private/CineAssemblyToolsStyle.cpp b/Engine/Plugins/VirtualProduction/CinematicAssemblyTools/Source/CineAssemblyToolsEditor/Private/CineAssemblyToolsStyle.cpp
--- a/Engine/Plugins/VirtualProduction/CinematicAssemblyTools/Source/CineAssemblyToolsEditor/Private/CineAssemblyToolsStyle.cpp
+++ b/Engine/Plugins/VirtualProduction/CinematicAssemblyTools/Source/CineAssemblyToolsEditor/Private/CineAssemblyToolsStyle.cpp
@@ -22,11 +22,14 @@ FCineAssemblyToolsStyle::FCineAssemblyToolsStyle()
 {
 	// Icons
 	const FVector2D Icon16x16(16.0f, 16.0f);
+	const FVector2D Icon64x64(64.0f, 64.0f);
 
 	const FString ContentDir = IPluginManager::Get().FindPlugin(TEXT("CinematicAssemblyTools"))->GetContentDir();
 	SetContentRoot(ContentDir);
 	SetCoreContentRoot(FPaths::EngineContentDir() / TEXT("Slate"));
 
+	Set("Thumbnails.Schema", new FSlateVectorImageBrush(RootToContentDir(TEXT("Slate/LevelSequenceActor_64.svg")), Icon64x64));
+
 	Set("Icons.RevisionControl", new FSlateVectorImageBrush(RootToContentDir(TEXT("Slate/RevisionControl.svg")), Icon16x16));
 	Set("Icons.Productions", new FSlateVectorImageBrush(RootToContentDir(TEXT("Slate/Camera.svg")), Icon16x16));
 	Set("Icons.Sequencer", new FSlateVectorImageBrush(RootToContentDir(TEXT("Slate/Sequencer_16.svg")), Icon16x16));
@@ -57,6 +60,7 @@ FCineAssemblyToolsStyle::FCineAssemblyToolsStyle()
 
 	Set("Borders.Background", new FSlateRoundedBoxBrush(FStyleColors::Background, 0.0f, FStyleColors::Panel, 2.0f));
 	Set("Borders.PanelNoBorder", new FSlateRoundedBoxBrush(FStyleColors::Panel, 0.0f, FStyleColors::Panel, 0.0f));
+	Set("Borders.RecessedNoBorder", new FSlateRoundedBoxBrush(FStyleColors::Recessed, 0.0f, FStyleColors::Recessed, 0.0f));
 
 	// Buttons
 	const FButtonStyle RecessedButton = FButtonStyle()
diff --git a/Engine/Plugins/VirtualProduction/CinematicAssemblyTools/Source/CineAssemblyToolsEditor/Private/UI/CineAssembly/SCineAssemblyConfigWindow.cpp b/Engine/Plugins/VirtualProduction/CinematicAssemblyTools/Source/CineAssemblyToolsEditor/Private/UI/CineAssembly/SCineAssemblyConfigWindow.cpp
--- a/Engine/Plugins/VirtualProduction/CinematicAssemblyTools/Source/CineAssemblyToolsEditor/Private/UI/CineAssembly/SCineAssemblyConfigWindow.cpp
+++ b/Engine/Plugins/VirtualProduction/CinematicAssemblyTools/Source/CineAssemblyToolsEditor/Private/UI/CineAssembly/SCineAssemblyConfigWindow.cpp
@@ -13,6 +13,7 @@
 #include "IStructureDetailsView.h"
 #include "Misc/ConfigCacheIni.h"
 #include "Modules/ModuleManager.h"
+#include "ObjectTools.h"
 #include "PropertyEditorModule.h"
 #include "UObject/Package.h"
 #include "Settings/ContentBrowserSettings.h"
@@ -30,7 +31,7 @@
 
 namespace UE::CineAssemblyWidgets::Private
 {
-	TSharedRef<IDetailsView> MakeDetailsWidget(UCineAssembly* Assembly)
+	TSharedRef<IDetailsView> MakeDetailsView(UCineAssembly* Assembly)
 	{
 		FPropertyEditorModule& PropertyEditor = FModuleManager::Get().LoadModuleChecked<FPropertyEditorModule>(TEXT("PropertyEditor"));
 
@@ -218,14 +219,9 @@ void SCineAssemblyConfigWindow::OnSchemaSelected(const FAssetData& AssetData)
 
 TSharedRef<SWidget> SCineAssemblyConfigWindow::MakeInfoPanel()
 {
-	DetailsView = UE::CineAssemblyWidgets::Private::MakeDetailsWidget(CineAssemblyToConfigure.Get());
-
 	TabSwitcher = SNew(SWidgetSwitcher)
 		+ SWidgetSwitcher::Slot()
-			[ MakeOverviewWidget() ]
-
-		+ SWidgetSwitcher::Slot()
-			[ DetailsView.ToSharedRef() ]
+			[ MakeDetailsWidget() ]
 
 		+ SWidgetSwitcher::Slot()
 			[ MakeHierarchyWidget() ]
@@ -246,18 +242,14 @@ TSharedRef<SWidget> SCineAssemblyConfigWindow::MakeInfoPanel()
 				})
 
 				+ SSegmentedControl<int32>::Slot(0)
-				.Text(LOCTEXT("OverviewTab", "Overview"))
-				.Icon(FSlateIcon(FCineAssemblyToolsStyle::StyleName, "Icons.Animation").GetIcon())
-
-				+ SSegmentedControl<int32>::Slot(1)
 				.Text(LOCTEXT("DetailsTab", "Details"))
 				.Icon(FSlateIcon(FAppStyle::GetAppStyleSetName(), "Icons.Details").GetIcon())
 
-				+ SSegmentedControl<int32>::Slot(2)
+				+ SSegmentedControl<int32>::Slot(1)
 				.Text(LOCTEXT("HierarchyTab", "Hierarchy"))
 				.Icon(FSlateIcon(FAppStyle::GetAppStyleSetName(), "Icons.FolderClosed").GetIcon())
 
-				+ SSegmentedControl<int32>::Slot(3)
+				+ SSegmentedControl<int32>::Slot(2)
 				.Text(LOCTEXT("NotesTab", "Notes"))
 				.Icon(FSlateIcon(FCineAssemblyToolsStyle::StyleName, "Icons.Notes").GetIcon())
 		]
@@ -320,73 +312,83 @@ void SCineAssemblyConfigWindow::EvaluateTokenString(FTemplateString& StringToEva
 	}
 }
 
-TSharedRef<SWidget> SCineAssemblyConfigWindow::MakeOverviewWidget()
+TSharedRef<SWidget> SCineAssemblyConfigWindow::MakeDetailsWidget()
 {
-	return SNew(SBorder)
-		.BorderImage(FCineAssemblyToolsStyle::Get().GetBrush("ProductionWizard.PanelBackground"))
-		.Padding(16.0f)
-		[
-			SNew(SVerticalBox)
+	DetailsView = UE::CineAssemblyWidgets::Private::MakeDetailsView(CineAssemblyToConfigure.Get());
 
-			// TODO: Populate this with an actual preview image sourced from the schema
-			+ SVerticalBox::Slot()
-				.Padding(0.0f, 0.0f, 0.0f, 16.0f)
-				.MinHeight(300.0f)
-				.MaxHeight(300.0f)
-				[
-					SNew(SBorder)
-						.BorderImage(FCineAssemblyToolsStyle::Get().GetBrush("Borders.Background"))
-						.HAlign(HAlign_Center)
-						.VAlign(VAlign_Center)
-						[
-							SNew(STextBlock).Text(LOCTEXT("PreviewImageText", "PREVIEW IMAGE"))
-						]
-				]
+	return SNew(SVerticalBox)
 
-			+ SVerticalBox::Slot()
-				.AutoHeight()
-				.Padding(0.0f, 0.0f, 0.0f, 4.0f)
+		+ SVerticalBox::Slot()
+		.AutoHeight()
+		[
+			SNew(SBorder)
+				.BorderImage(FCineAssemblyToolsStyle::Get().GetBrush("Borders.RecessedNoBorder"))
+				.Padding(16.0f)
 				[
-					SNew(STextBlock)
-						.AutoWrapText(true)
-						.Text_Lambda([this]() -> FText
-						{
-							const FText NoSchemaText = LOCTEXT("SchemaInstructions", "Choose a schema to use as the base for configuring your Cine Assembly, or proceed with no schema.");
-							return SelectedSchema ? FText::FromString(SelectedSchema->SchemaName) : NoSchemaText;
-						})
-				]
+					SNew(SVerticalBox)
 
-			+ SVerticalBox::Slot()
-				.AutoHeight()
-				.Padding(0.0f, 0.0f, 0.0f, 8.0f)
-				[
-					SNew(STextBlock)
-						.Text(LOCTEXT("SchemClassName", "Cine Assembly Schema"))
-						.ColorAndOpacity(FSlateColor::UseSubduedForeground())
-						.Visibility_Lambda([this]() -> EVisibility 
-						{
-							return SelectedSchema ? EVisibility::Visible : EVisibility::Collapsed;
-						})
-				]
+					+ SVerticalBox::Slot()
+						.Padding(0.0f, 0.0f, 0.0f, 16.0f)
+						.AutoHeight()
+						[
+							SNew(SHorizontalBox)
 
-			+ SVerticalBox::Slot()
-				.AutoHeight()
-				[
-					SNew(STextBlock)
-						.Text_Lambda([this]() -> FText
-						{
-							if (SelectedSchema)
-							{
-								return !SelectedSchema->Description.IsEmpty() ? FText::FromString(SelectedSchema->Description) : LOCTEXT("EmptyDescription", "No description");
-							}
-							return FText::GetEmpty();
-						})
-						.Visibility_Lambda([this]() -> EVisibility
-						{
-							return SelectedSchema ? EVisibility::Visible : EVisibility::Collapsed;
-						})
+							+ SHorizontalBox::Slot()
+								.AutoWidth()
+								.VAlign(VAlign_Center)
+								.Padding(0.0f, 0.0f, 8.0f, 0.0f)
+								[
+									SNew(SImage)
+										.Image_Lambda([this]()
+											{
+												return SelectedSchema ? SelectedSchema->GetThumbnailBrush() : FCineAssemblyToolsStyle::Get().GetBrush("Thumbnails.Schema");
+											})
+								]
 
+								+ SHorizontalBox::Slot()
+								.VAlign(VAlign_Center)
+								[
+									SNew(SVerticalBox)
+
+									+ SVerticalBox::Slot()
+										[
+											SNew(STextBlock)
+												.Text_Lambda([this]()
+													{
+														return SelectedSchema ? FText::FromString(SelectedSchema->SchemaName) : LOCTEXT("NoSchemaName", "No Schema");
+													})
+										]
+
+									+ SVerticalBox::Slot()
+										[
+											SNew(STextBlock)
+												.Text(LOCTEXT("SchemClassName", "Cine Assembly Schema"))
+												.ColorAndOpacity(FSlateColor::UseSubduedForeground())
+										]
+								]
+						]
+
+					+ SVerticalBox::Slot()
+						.AutoHeight()
+						[
+							SNew(STextBlock)
+								.AutoWrapText(true)
+								.Text_Lambda([this]()
+									{
+										if (SelectedSchema)
+										{
+											return !SelectedSchema->Description.IsEmpty() ? FText::FromString(SelectedSchema->Description) : LOCTEXT("EmptyDescription", "No description");
+										}
+										return LOCTEXT("SchemaInstructions", "Choose a schema to use as the base for configuring your Cine Assembly, or proceed with no schema.");
+									})
+						]
 				]
+		]
+
+		+ SVerticalBox::Slot()
+		.FillContentHeight(1.0f)
+		[
+			DetailsView.ToSharedRef()
 		];
 }
 
@@ -732,7 +734,7 @@ TSharedRef<SWidget> SCineAssemblyEditWidget::BuildUI()
 		[MakeOverviewWidget()]
 
 		+ SWidgetSwitcher::Slot()
-		[UE::CineAssemblyWidgets::Private::MakeDetailsWidget(CineAssembly)];
+		[UE::CineAssemblyWidgets::Private::MakeDetailsView(CineAssembly)];
 
 	return SNew(SBorder)
 		.BorderImage(FCineAssemblyToolsStyle::Get().GetBrush("Borders.PanelNoBorder"))
@@ -803,6 +805,22 @@ FString SCineAssemblyEditWidget::GetAssemblyName()
 	return TEXT("CineAssembly");
 }
 
+bool SCineAssemblyEditWidget::HasRenderedThumbnail()
+{
+	const FObjectThumbnail* ObjectThumbnail = nullptr;
+	if (CineAssembly)
+	{
+		const FName FullAssetName = *CineAssembly->GetFullName();
+
+		FThumbnailMap ThumbnailMap;
+		ThumbnailTools::ConditionallyLoadThumbnailsForObjects({ FullAssetName }, ThumbnailMap);
+		ObjectThumbnail = ThumbnailMap.Find(FullAssetName);
+	}
+
+	const bool bHasRenderedThumbnail = ObjectThumbnail && !ObjectThumbnail->IsEmpty();
+	return bHasRenderedThumbnail;
+}
+
 TSharedRef<SWidget> SCineAssemblyEditWidget::MakeOverviewWidget()
 {
 	const UCineAssemblySchema* Schema = CineAssembly->GetSchema();
@@ -812,60 +830,65 @@ TSharedRef<SWidget> SCineAssemblyEditWidget::MakeOverviewWidget()
 			return Schema ? EVisibility::Visible : EVisibility::Collapsed;
 		};
 
+	TSharedRef<FAssetThumbnail> AssetThumbnail = MakeShared<FAssetThumbnail>(CineAssembly, 256, 256, UThumbnailManager::Get().GetSharedThumbnailPool());
+	FAssetThumbnailConfig ThumbnailConfig;
+	ThumbnailConfig.AssetTypeColorOverride = FLinearColor::Transparent;
+
 	return SNew(SVerticalBox)
 
-		// TODO: Populate this with an actual preview image sourced from the schema
 		+ SVerticalBox::Slot()
+		.AutoHeight()
 		.Padding(0.0f, 0.0f, 0.0f, 16.0f)
-		.MinHeight(300.0f)
-		.MaxHeight(300.0f)
 		[
 			SNew(SBorder)
-				.BorderImage(FCineAssemblyToolsStyle::Get().GetBrush("Borders.Background"))
-				.HAlign(HAlign_Center)
-				.VAlign(VAlign_Center)
+				.BorderImage(FCineAssemblyToolsStyle::Get().GetBrush("ProductionWizard.RecessedBackground"))
+				.Padding(4.0f)
 				[
-					SNew(STextBlock).Text(LOCTEXT("PreviewImageText", "PREVIEW IMAGE"))
-				]
-		]
+					SNew(SVerticalBox)
 
-		+ SVerticalBox::Slot()
-		.AutoHeight()
-		.Padding(0.0f, 0.0f, 0.0f, 4.0f)
-		[
-			SNew(STextBlock)
-				.AutoWrapText(true)
-				.Text_Lambda([Schema]() -> FText
-					{
-						return Schema ? FText::FromString(Schema->SchemaName) : FText::GetEmpty();
-					})
-				.Visibility_Lambda(GetVisbility)
-		]
+					+SVerticalBox::Slot()
+						.AutoHeight()
+						.Padding(0.0f, 8.0f, 0.0f, 8.0f)
+						.HAlign(HAlign_Center)
+						[
+							SNew(SHorizontalBox)
 
-		+ SVerticalBox::Slot()
-		.AutoHeight()
-		.Padding(0.0f, 0.0f, 0.0f, 8.0f)
-		[
-			SNew(STextBlock)
-				.Text(LOCTEXT("SchemClassName", "Cine Assembly Schema"))
-				.ColorAndOpacity(FSlateColor::UseSubduedForeground())
-				.Visibility_Lambda(GetVisbility)
-		]
+								+ SHorizontalBox::Slot()
+								.AutoWidth()
+								[
+									AssetThumbnail->MakeThumbnailWidget(ThumbnailConfig)
+								]
 
-		+ SVerticalBox::Slot()
-		.AutoHeight()
-		.Padding(0.0f, 0.0f, 0.0f, 16.0f)
-		[
-			SNew(STextBlock)
-				.Text_Lambda([Schema]() -> FText
-					{
-						if (Schema)
-						{
-							return !Schema->Description.IsEmpty() ? FText::FromString(Schema->Description) : LOCTEXT("EmptyDescription", "No description");
-						}
-						return FText::GetEmpty();
-					})
-				.Visibility_Lambda(GetVisbility)
+								+ SHorizontalBox::Slot()
+								.Padding(16.0f, 0.0, 0.0, 0.0f)
+								.VAlign(VAlign_Center)
+								.FillContentWidth(1.0f)
+								[
+									SNew(STextBlock)
+										.Text(LOCTEXT("ThumbnailHintText", "This assembly does not currently have a preview thumbnail. Open this asset in Sequencer and save it to render a preview to display here."))
+										.AutoWrapText(true)
+										.Visibility_Lambda([this]() { return HasRenderedThumbnail() ? EVisibility::Collapsed : EVisibility::Visible; })
+								]
+						]
+
+						+ SVerticalBox::Slot()
+						.AutoHeight()
+						.HAlign(HAlign_Center)
+						.Padding(0.0f, 0.0f, 0.0f, 4.0f)
+						[
+							SNew(STextBlock).Text_Lambda([this]() {	return FText::FromString(GetAssemblyName()); })
+						]
+
+						+ SVerticalBox::Slot()
+						.AutoHeight()
+						.HAlign(HAlign_Center)
+						.Padding(0.0f, 0.0f, 0.0f, 8.0f)
+						[
+							SNew(STextBlock)
+								.Text_Lambda([Schema]() { return Schema ? FText::FromString(Schema->SchemaName) : LOCTEXT("NoSchemaName", "No Schema"); })
+								.ColorAndOpacity(FSlateColor::UseSubduedForeground())
+						]
+				]
 		]
 
 		+ SVerticalBox::Slot()
diff --git a/Engine/Plugins/VirtualProduction/CinematicAssemblyTools/Source/CineAssemblyToolsEditor/Private/UI/CineAssembly/SCineAssemblyConfigWindow.h b/Engine/Plugins/VirtualProduction/CinematicAssemblyTools/Source/CineAssemblyToolsEditor/Private/UI/CineAssembly/SCineAssemblyConfigWindow.h
--- a/Engine/Plugins/VirtualProduction/CinematicAssemblyTools/Source/CineAssemblyToolsEditor/Private/UI/CineAssembly/SCineAssemblyConfigWindow.h
+++ b/Engine/Plugins/VirtualProduction/CinematicAssemblyTools/Source/CineAssemblyToolsEditor/Private/UI/CineAssembly/SCineAssemblyConfigWindow.h
@@ -58,7 +58,7 @@ private:
 	TSharedRef<SWidget> MakeButtonsPanel();
 
 	/** Creates the widget to display for the Overview tab */
-	TSharedRef<SWidget> MakeOverviewWidget();
+	TSharedRef<SWidget> MakeDetailsWidget();
 
 	/** Creates the widget to display for the Hierarchy tab */
 	TSharedRef<SWidget> MakeHierarchyWidget();
@@ -161,6 +161,9 @@ public:
 	/** Searches the asset registry for a Cine Assembly matching the input ID and updates the UI */
 	void FindAssembly(FGuid AssemblyID);
 
+	/** Returns true if the Assembly asset has a rendered thumbnail (such as from the Sequencer preview) */
+	bool HasRenderedThumbnail();
+
 private:
 	/** Constructs the main UI for the widget */ 
 	TSharedRef<SWidget> BuildUI();
