ShotManagement: Remove unnecessary overview panel from Cine Assembly Schema creation window.
#rnx

diff --git a/Engine/Plugins/VirtualProduction/CinematicAssemblyTools/Source/CineAssemblyToolsEditor/Private/UI/CineAssembly/SCineAssemblySchemaWindow.cpp b/Engine/Plugins/VirtualProduction/CinematicAssemblyTools/Source/CineAssemblyToolsEditor/Private/UI/CineAssembly/SCineAssemblySchemaWindow.cpp
--- a/Engine/Plugins/VirtualProduction/CinematicAssemblyTools/Source/CineAssemblyToolsEditor/Private/UI/CineAssembly/SCineAssemblySchemaWindow.cpp
+++ b/Engine/Plugins/VirtualProduction/CinematicAssemblyTools/Source/CineAssemblyToolsEditor/Private/UI/CineAssembly/SCineAssemblySchemaWindow.cpp
@@ -148,12 +148,17 @@ TSharedRef<SWidget> SCineAssemblySchemaWindow::MakeMenuPanel()
 				})
 
 				+ SSegmentedControl<int32>::Slot(0)
-				.Text(LOCTEXT("AssetTab", "Asset"))
-				.Icon(FSlateIcon(FCineAssemblyToolsStyle::StyleName, "Icons.Shot").GetIcon())
+				.Text(LOCTEXT("DetailsTab", "Details"))
+				.Icon(FSlateIcon(FAppStyle::GetAppStyleSetName(), "Icons.Details").GetIcon())
 				.HAlign(HAlign_Left)
 
 				+ SSegmentedControl<int32>::Slot(1)
-				.Text(LOCTEXT("HierarchyTab", "Hierarchy"))
+				.Text(LOCTEXT("MetadataTab", "Metadata"))
+				.Icon(FSlateIcon(FCineAssemblyToolsStyle::StyleName, "Icons.DataAsset").GetIcon())
+				.HAlign(HAlign_Left)
+
+				+ SSegmentedControl<int32>::Slot(2)
+				.Text(LOCTEXT("HierarchyTab", "Content Hierarchy"))
 				.Icon(FSlateIcon(FAppStyle::GetAppStyleSetName(), "Icons.FolderClosed").GetIcon())
 				.HAlign(HAlign_Left)
 		];
@@ -164,7 +169,12 @@ TSharedRef<SWidget> SCineAssemblySchemaWindow::MakeContentPanel()
 	MenuTabSwitcher = SNew(SWidgetSwitcher)
 		+ SWidgetSwitcher::Slot()
 		[
-			MakeAssetTabContent()
+			MakeDetailsTabContent()
+		]
+
+		+ SWidgetSwitcher::Slot()
+		[
+			MakeMetadataTabContent()
 		]
 
 		+ SWidgetSwitcher::Slot()
@@ -262,143 +272,104 @@ TSharedRef<SWidget> SCineAssemblySchemaWindow::MakeButtonsPanel()
 	return SNullWidget::NullWidget;
 }
 
-TSharedRef<SWidget> SCineAssemblySchemaWindow::MakeAssetTabContent()
+TSharedRef<SWidget> SCineAssemblySchemaWindow::MakeDetailsWidget(bool bShowMetadata)
 {
-	AssetTabSwitcher = SNew(SWidgetSwitcher)
-		+ SWidgetSwitcher::Slot()
-		[ MakeOverviewWidget() ]
-
-		+ SWidgetSwitcher::Slot()
-		[ MakeDetailsWidget(false) ]
-
-		+ SWidgetSwitcher::Slot()
-		[ MakeDetailsWidget(true) ];
+	FPropertyEditorModule& PropertyEditor = FModuleManager::Get().LoadModuleChecked<FPropertyEditorModule>(TEXT("PropertyEditor"));
 
-	return SNew(SVerticalBox)
+	FDetailsViewArgs DetailsViewArgs;
+	DetailsViewArgs.bAllowSearch = false;
+	DetailsViewArgs.bHideSelectionTip = true;
 
-		+ SVerticalBox::Slot()
-		.AutoHeight()
-		.HAlign(HAlign_Center)
-		.Padding(8.0f)
-		[
-			SNew(SSegmentedControl<int32>)
-				.Value(0)
-				.OnValueChanged_Lambda([Switcher = AssetTabSwitcher](int32 NewValue)
-					{
-						Switcher->SetActiveWidgetIndex(NewValue);
-					})
+	TSharedRef<IDetailsView> DetailsView = PropertyEditor.CreateDetailView(DetailsViewArgs);
 
-				+ SSegmentedControl<int32>::Slot(0)
-				.Text(LOCTEXT("OverviewTab", "Overview"))
-				.Icon(FSlateIcon(FCineAssemblyToolsStyle::StyleName, "Icons.Animation").GetIcon())
+	DetailsView->SetIsPropertyVisibleDelegate(FIsPropertyVisible::CreateSP(this, &SCineAssemblySchemaWindow::IsPropertyVisible, bShowMetadata));
 
-				+ SSegmentedControl<int32>::Slot(1)
-				.Text(LOCTEXT("DetailsTab", "Details"))
-				.Icon(FSlateIcon(FAppStyle::GetAppStyleSetName(), "Icons.Details").GetIcon())
+	DetailsView->SetObject(SchemaToConfigure.Get(), true);
+	DetailsView->OnFinishedChangingProperties().AddSP(this, &SCineAssemblySchemaWindow::OnSchemaPropertiesChanged);
 
-				+ SSegmentedControl<int32>::Slot(2)
-				.Text(LOCTEXT("DataTab", "Data"))
-				.Icon(FSlateIcon(FCineAssemblyToolsStyle::StyleName, "Icons.DataAsset").GetIcon())
-		]
+	return DetailsView;
+}
 
-		+ SVerticalBox::Slot()
-		.AutoHeight()
-		.HAlign(HAlign_Fill)
-		[
-			AssetTabSwitcher.ToSharedRef()
-		];
+bool SCineAssemblySchemaWindow::IsPropertyVisible(const FPropertyAndParent& PropertyAndParent, bool bShowMetadata)
+{
+	if (PropertyAndParent.Property.GetFName() == GET_MEMBER_NAME_CHECKED(UCineAssemblySchema, AssemblyMetadata))
+	{
+		return bShowMetadata;
+	}
+	else if ((PropertyAndParent.ParentProperties.Num() > 0) && (PropertyAndParent.ParentProperties[0]->GetFName() == GET_MEMBER_NAME_CHECKED(UCineAssemblySchema, AssemblyMetadata)))
+	{
+		return bShowMetadata;
+	}
+	return !bShowMetadata;
 }
 
-TSharedRef<SWidget> SCineAssemblySchemaWindow::MakeOverviewWidget()
+TSharedRef<SWidget> SCineAssemblySchemaWindow::MakeDetailsTabContent()
 {
 	return SNew(SBorder)
-		.BorderImage(FCineAssemblyToolsStyle::Get().GetBrush("Borders.PanelNoBorder"))
 		.Padding(16.0f)
+		.BorderImage(FCineAssemblyToolsStyle::Get().GetBrush("ProductionWizard.PanelBackground"))
 		[
 			SNew(SVerticalBox)
 
-				// TODO: Populate this with an actual preview image sourced from the schema
 				+ SVerticalBox::Slot()
-				.Padding(0.0f, 0.0f, 0.0f, 16.0f)
 				.AutoHeight()
-				.HAlign(HAlign_Center)
+				.Padding(0.0f, 0.0f, 0.0f, 8.0f)
 				[
-					SNew(SBox)
-					.WidthOverride(420.0f)
-					.HeightOverride(300.0f)
-						[
-							SNew(SBorder)
-								.BorderImage(FCineAssemblyToolsStyle::Get().GetBrush("Borders.Background"))
-								.HAlign(HAlign_Center)
-								.VAlign(VAlign_Center)
-								[
-									SNew(STextBlock).Text(LOCTEXT("PreviewImageText", "PREVIEW IMAGE"))
-								]
-						]
+					SNew(STextBlock)
+						.Text(LOCTEXT("DetailsTitle", "Schema Details"))
+						.Font(FCineAssemblyToolsStyle::Get().GetFontStyle("ProductionWizard.HeadingFont"))
 				]
-	
+
 				+ SVerticalBox::Slot()
 				.AutoHeight()
-				.Padding(0.0f, 0.0f, 0.0f, 4.0f)
+				.Padding(0.0f, 0.0f, 0.0f, 24.0f)
 				[
 					SNew(STextBlock)
 						.AutoWrapText(true)
-						.Text_Lambda([this]() -> FText
-							{
-								return FText::FromString(SchemaToConfigure->SchemaName);
-							})
+						.Text(LOCTEXT("SchemaDetailsInstruction", "Configure the properties which will be inherited by every Cine Assembly asset created from this schema."))
 				]
 
+				+ SVerticalBox::Slot()
+				.FillHeight(1.0f)
+				[
+					MakeDetailsWidget(false)
+				]
+		];
+}
+
+TSharedRef<SWidget> SCineAssemblySchemaWindow::MakeMetadataTabContent()
+{
+	return SNew(SBorder)
+		.Padding(16.0f)
+		.BorderImage(FCineAssemblyToolsStyle::Get().GetBrush("ProductionWizard.PanelBackground"))
+		[
+			SNew(SVerticalBox)
+
 				+ SVerticalBox::Slot()
 				.AutoHeight()
 				.Padding(0.0f, 0.0f, 0.0f, 8.0f)
 				[
 					SNew(STextBlock)
-						.Text(LOCTEXT("SchemClassName", "Cine Assembly Schema"))
-						.ColorAndOpacity(FSlateColor::UseSubduedForeground())
+						.Text(LOCTEXT("MetadataTitle", "Schema Metadata"))
+						.Font(FCineAssemblyToolsStyle::Get().GetFontStyle("ProductionWizard.HeadingFont"))
 				]
 
 				+ SVerticalBox::Slot()
 				.AutoHeight()
+				.Padding(0.0f, 0.0f, 0.0f, 24.0f)
 				[
 					SNew(STextBlock)
-						.Text_Lambda([this]() -> FText
-							{
-								return !SchemaToConfigure->Description.IsEmpty() ? FText::FromString(SchemaToConfigure->Description) : LOCTEXT("EmptyDescription", "No description");
-							})
+						.AutoWrapText(true)
+						.Text(LOCTEXT("SchemaMetadataInstruction", "Configure the metadata that should be associated with Cine Assemblies created from this schema. "
+							"For each metadata field, choose the value type, metadata key, and optionally a default value."))
 				]
-		];
-}
-
-TSharedRef<SWidget> SCineAssemblySchemaWindow::MakeDetailsWidget(bool bShowMetadata)
-{
-	FPropertyEditorModule& PropertyEditor = FModuleManager::Get().LoadModuleChecked<FPropertyEditorModule>(TEXT("PropertyEditor"));
 
-	FDetailsViewArgs DetailsViewArgs;
-	DetailsViewArgs.bAllowSearch = false;
-	DetailsViewArgs.bHideSelectionTip = true;
-
-	TSharedRef<IDetailsView> DetailsView = PropertyEditor.CreateDetailView(DetailsViewArgs);
-
-	DetailsView->SetIsPropertyVisibleDelegate(FIsPropertyVisible::CreateSP(this, &SCineAssemblySchemaWindow::IsPropertyVisible, bShowMetadata));
-
-	DetailsView->SetObject(SchemaToConfigure.Get(), true);
-	DetailsView->OnFinishedChangingProperties().AddSP(this, &SCineAssemblySchemaWindow::OnSchemaPropertiesChanged);
-
-	return DetailsView;
-}
-
-bool SCineAssemblySchemaWindow::IsPropertyVisible(const FPropertyAndParent& PropertyAndParent, bool bShowMetadata)
-{
-	if (PropertyAndParent.Property.GetFName() == GET_MEMBER_NAME_CHECKED(UCineAssemblySchema, AssemblyMetadata))
-	{
-		return bShowMetadata;
-	}
-	else if ((PropertyAndParent.ParentProperties.Num() > 0) && (PropertyAndParent.ParentProperties[0]->GetFName() == GET_MEMBER_NAME_CHECKED(UCineAssemblySchema, AssemblyMetadata)))
-	{
-		return bShowMetadata;
-	}
-	return !bShowMetadata;
+				+ SVerticalBox::Slot()
+				.FillHeight(1.0f)
+				[
+					MakeDetailsWidget(true)
+				]
+		];
 }
 
 TSharedRef<SWidget> SCineAssemblySchemaWindow::MakeHierarchyTabContent()
@@ -414,7 +385,7 @@ TSharedRef<SWidget> SCineAssemblySchemaWindow::MakeHierarchyTabContent()
 
 	ExpandTreeRecursive(RootItem);
 
-	TSharedRef<SWidget> ContentHierarhcyWidget = SNew(SVerticalBox)
+	TSharedRef<SWidget> ContentHierarchyWidget = SNew(SVerticalBox)
 
 		+ SVerticalBox::Slot()
 		.AutoHeight()
@@ -548,7 +519,7 @@ TSharedRef<SWidget> SCineAssemblySchemaWindow::MakeHierarchyTabContent()
 						.FillWidth(0.5f)
 						.Padding(16.0f, 0.0f, 0.0f, 0.0f)
 						[
-							ContentHierarhcyWidget
+							ContentHierarchyWidget
 						]
 				]
 		];
diff --git a/Engine/Plugins/VirtualProduction/CinematicAssemblyTools/Source/CineAssemblyToolsEditor/Private/UI/CineAssembly/SCineAssemblySchemaWindow.h b/Engine/Plugins/VirtualProduction/CinematicAssemblyTools/Source/CineAssemblyToolsEditor/Private/UI/CineAssembly/SCineAssemblySchemaWindow.h
--- a/Engine/Plugins/VirtualProduction/CinematicAssemblyTools/Source/CineAssemblyToolsEditor/Private/UI/CineAssembly/SCineAssemblySchemaWindow.h
+++ b/Engine/Plugins/VirtualProduction/CinematicAssemblyTools/Source/CineAssemblyToolsEditor/Private/UI/CineAssembly/SCineAssemblySchemaWindow.h
@@ -110,15 +110,15 @@ private:
 	/** Creates the buttons on the bottom of the window */
 	TSharedRef<SWidget> MakeButtonsPanel();
 
-	/** Creates the content for the Assets tab */
-	TSharedRef<SWidget> MakeAssetTabContent();
+	/** Creates the content for the Details tab */
+	TSharedRef<SWidget> MakeDetailsTabContent();
+
+	/** Creates the content for the Metadata tab */
+	TSharedRef<SWidget> MakeMetadataTabContent();
 
 	/** Creates the content for the Hierarchy tab */
 	TSharedRef<SWidget> MakeHierarchyTabContent();
 
-	/** Creates the content for the Overview tab */
-	TSharedRef<SWidget> MakeOverviewWidget();
-
 	/** Creates the content for the Details and Data tabs (metadata properties are only shown in the Data tab) */
 	TSharedRef<SWidget> MakeDetailsWidget(bool bShowMetadata);
 
