ShotManagement: Add additional validation to Schema metadata keys to prevent empty keys, duplicate keys, and keys that match the default CAT tokens.
#rnx

diff --git a/Engine/Plugins/VirtualProduction/CinematicAssemblyTools/Source/CineAssemblyToolsEditor/Private/CineAssemblyMetadataCustomization.cpp b/Engine/Plugins/VirtualProduction/CinematicAssemblyTools/Source/CineAssemblyToolsEditor/Private/CineAssemblyMetadataCustomization.cpp
--- a/Engine/Plugins/VirtualProduction/CinematicAssemblyTools/Source/CineAssemblyToolsEditor/Private/CineAssemblyMetadataCustomization.cpp
+++ b/Engine/Plugins/VirtualProduction/CinematicAssemblyTools/Source/CineAssemblyToolsEditor/Private/CineAssemblyMetadataCustomization.cpp
@@ -2,18 +2,24 @@
 
 #include "CineAssemblyMetadataCustomization.h"
 
+#include "Algo/Contains.h"
+#include "Algo/IndexOf.h"
 #include "CineAssembly.h"
-#include "CineAssemblySchema.h"
+#include "CineAssemblyNamingTokens.h"
 #include "IDetailChildrenBuilder.h"
 #include "IPropertyUtilities.h"
+#include "NamingTokensEngineSubsystem.h"
 #include "PropertyCustomizationHelpers.h"
 #include "Widgets/Input/SCheckBox.h"
+#include "Widgets/Input/SEditableTextBox.h"
 #include "Widgets/Input/SMultiLineEditableTextBox.h"
 #include "Widgets/Input/SNumericEntryBox.h"
 #include "Widgets/Layout/SBox.h"
 #include "Widgets/Layout/SWidgetSwitcher.h"
 #include "Widgets/Text/STextBlock.h"
 
+#define LOCTEXT_NAMESPACE "CineAssemblyMetadataCustomization"
+
 TSharedRef<IPropertyTypeCustomization> FCineAssemblyMetadataCustomization::MakeInstance()
 {
 	return MakeShareable(new FCineAssemblyMetadataCustomization);
@@ -40,24 +46,36 @@ void FCineAssemblyMetadataCustomization::CustomizeChildren(TSharedRef<IPropertyH
 		return;
 	}
 
-	UCineAssemblySchema* Schema = Cast<UCineAssemblySchema>(CustomizedObjects[0].Get());
-
-	// Add all of the existing reflected properties of the metadata struct
-	uint32 NumChildren;
-	PropertyHandle->GetNumChildren(NumChildren);
+	Schema = Cast<UCineAssemblySchema>(CustomizedObjects[0].Get());
 
-	for (uint32 Index = 0; Index < NumChildren; ++Index)
+	if (!PropertyHandle->IsExpanded())
 	{
-		TSharedRef<IPropertyHandle> ChildPropertyHandle = PropertyHandle->GetChildHandle(Index).ToSharedRef();
-		ChildBuilder.AddProperty(ChildPropertyHandle).ShowPropertyButtons(false);
+		PropertyHandle->SetExpanded(true);
 	}
 
-	// Add a "Default Value" property, based on the metadata type
-	int32 ArrayIndex = PropertyHandle->GetArrayIndex();
+	ArrayIndex = PropertyHandle->GetArrayIndex();
 	if (Schema->AssemblyMetadata.IsValidIndex(ArrayIndex))
 	{
 		FAssemblyMetadataDesc& MetadataDesc = Schema->AssemblyMetadata[ArrayIndex];
 
+		// Add all of the existing reflected properties of the metadata struct
+		uint32 NumChildren;
+		PropertyHandle->GetNumChildren(NumChildren);
+
+		for (uint32 Index = 0; Index < NumChildren; ++Index)
+		{
+			TSharedRef<IPropertyHandle> ChildPropertyHandle = PropertyHandle->GetChildHandle(Index).ToSharedRef();
+
+			if (ChildPropertyHandle->GetProperty()->GetFName() == GET_MEMBER_NAME_CHECKED(FAssemblyMetadataDesc, Key))
+			{
+				CustomizeKeyProperty(ChildPropertyHandle, ChildBuilder, MetadataDesc);
+			}
+			else
+			{
+				ChildBuilder.AddProperty(ChildPropertyHandle).ShowPropertyButtons(false);
+			}
+		}
+
 		// Create a widget switcher that can display the appropriate widget based on the metadata type
 		TSharedPtr<IPropertyHandle> MetadataTypeHandle = PropertyHandle->GetChildHandle(GET_MEMBER_NAME_CHECKED(FAssemblyMetadataDesc, Type));
 
@@ -83,6 +101,7 @@ void FCineAssemblyMetadataCustomization::CustomizeChildren(TSharedRef<IPropertyH
 				}
 			}));
 
+		// Add a "Default Value" property, based on the metadata type
 		TSharedPtr<SWidgetSwitcher> DefaultValueWidget = SNew(SWidgetSwitcher)
 			.WidgetIndex_Lambda([MetadataTypeHandle]()
 				{
@@ -221,3 +240,87 @@ void FCineAssemblyMetadataCustomization::CustomizeChildren(TSharedRef<IPropertyH
 			];
 	}
 }
+
+void FCineAssemblyMetadataCustomization::CustomizeKeyProperty(TSharedRef<IPropertyHandle> PropertyHandle, IDetailChildrenBuilder& ChildBuilder, FAssemblyMetadataDesc& MetadataDesc)
+{
+	FString ExistingKeyName;
+	PropertyHandle->GetValue(ExistingKeyName);
+
+	// If the key name is not yet set, assign it a unique default key name
+	if (ExistingKeyName.IsEmpty())
+	{
+		PropertyHandle->SetValue(MakeUniqueKeyName());
+	}
+
+	IDetailPropertyRow& KeyRow = ChildBuilder.AddProperty(PropertyHandle);
+	KeyRow.CustomWidget()
+		.NameContent()
+		[
+			PropertyHandle->CreatePropertyNameWidget()
+		]
+		.ValueContent()
+		[
+			SNew(SEditableTextBox)
+				.Font(FAppStyle::GetFontStyle(TEXT("PropertyWindow.NormalFont")))
+				.Text_Lambda([this, PropertyHandle]() -> FText
+					{
+						FString KeyName;
+						PropertyHandle->GetValue(KeyName);
+						return FText::FromString(KeyName);
+					})
+				.OnTextCommitted_Lambda([this, PropertyHandle](const FText& InText, ETextCommit::Type InCommitType)
+					{
+						const FString NewKeyName = InText.ToString();
+						PropertyHandle->SetValue(NewKeyName);
+					})
+				.OnVerifyTextChanged_Raw(this, &FCineAssemblyMetadataCustomization::ValidateKeyName)
+		];
+}
+
+bool FCineAssemblyMetadataCustomization::ValidateKeyName(const FText& InText, FText& OutErrorMessage) const
+{
+	// An empty name is invalid
+	if (InText.IsEmpty())
+	{
+		OutErrorMessage = LOCTEXT("EmptyKeyNameError", "Please provide a key name");
+		return false;
+	}
+
+	// Check for duplicate keys in this schema
+	const int32 MetadataIndex = Algo::IndexOfBy(Schema->AssemblyMetadata, InText.ToString(), &FAssemblyMetadataDesc::Key);
+	if ((MetadataIndex != INDEX_NONE) && (MetadataIndex != ArrayIndex))
+	{
+		OutErrorMessage = LOCTEXT("DuplicateKeyNameError", "A metadata key with this name already exists in this schema");
+		return false;
+	}
+
+	// Check that the proposed key name does not match one of the default CAT tokens
+	UNamingTokensEngineSubsystem* NamingTokensSubsystem = GEngine->GetEngineSubsystem<UNamingTokensEngineSubsystem>();
+	UCineAssemblyNamingTokens* CineAssemblyNamingTokens = Cast<UCineAssemblyNamingTokens>(NamingTokensSubsystem->GetNamingTokens(UCineAssemblyNamingTokens::TokenNamespace));
+
+	TArray<FNamingTokenData> Tokens = CineAssemblyNamingTokens->GetDefaultTokens();
+	if (Algo::ContainsBy(Tokens, InText.ToString(), &FNamingTokenData::TokenKey))
+	{
+		OutErrorMessage = LOCTEXT("ExistingTokenKeyError", "A CAT token key with this name already exists");
+		return false;
+	}
+
+	return true;
+}
+
+FString FCineAssemblyMetadataCustomization::MakeUniqueKeyName()
+{
+	const FString BaseName = TEXT("NewKey");
+	FString WorkingName = BaseName;
+
+	int32 IntSuffix = 1;
+	while (Algo::ContainsBy(Schema->AssemblyMetadata, WorkingName, &FAssemblyMetadataDesc::Key))
+	{
+		WorkingName = FString::Printf(TEXT("%s%d"), *BaseName, IntSuffix);
+		IntSuffix++;
+	}
+
+	return WorkingName;
+}
+
+#undef LOCTEXT_NAMESPACE
diff --git a/Engine/Plugins/VirtualProduction/CinematicAssemblyTools/Source/CineAssemblyToolsEditor/Private/CineAssemblyMetadataCustomization.h b/Engine/Plugins/VirtualProduction/CinematicAssemblyTools/Source/CineAssemblyToolsEditor/Private/CineAssemblyMetadataCustomization.h
--- a/Engine/Plugins/VirtualProduction/CinematicAssemblyTools/Source/CineAssemblyToolsEditor/Private/CineAssemblyMetadataCustomization.h
+++ b/Engine/Plugins/VirtualProduction/CinematicAssemblyTools/Source/CineAssemblyToolsEditor/Private/CineAssemblyMetadataCustomization.h
@@ -4,6 +4,8 @@
 
 #include "IPropertyTypeCustomization.h"
 
+#include "CineAssemblySchema.h"
+
 class FCineAssemblyMetadataCustomization : public IPropertyTypeCustomization
 {
 public:
@@ -13,4 +15,21 @@ public:
 	virtual void CustomizeHeader(TSharedRef<IPropertyHandle> PropertyHandle, FDetailWidgetRow& HeaderRow, IPropertyTypeCustomizationUtils& CustomizationUtils) override;
 	virtual void CustomizeChildren(TSharedRef<IPropertyHandle> PropertyHandle, IDetailChildrenBuilder& ChildBuilder, IPropertyTypeCustomizationUtils& CustomizationUtils) override;
 	//~ End IDetailCustomization interface
+
+private:
+	/** Customizes the metadata desc Key property */
+	void CustomizeKeyProperty(TSharedRef<IPropertyHandle> PropertyHandle, IDetailChildrenBuilder& ChildBuilder, FAssemblyMetadataDesc& MetadataDesc);
+
+	/** Checks whether the proposed key name is valid (not empty, not a duplicate) */
+	bool ValidateKeyName(const FText& InText, FText& OutErrorMessage) const;
+
+	/** Returns a unique key name within the context of the schema */
+	FString MakeUniqueKeyName();
+
+private:
+	/** The schema that owns the metadata desc struct being customized */
+	UCineAssemblySchema* Schema;
+
+	/** The array index of the metadata struct being customized */
+	int32 ArrayIndex;
 };
