ShotManagement: Add a detail customization for the Cine Assembly Schema to include additional validation for the schema name and default assembly name properties
#rnx

diff --git a/Engine/Plugins/VirtualProduction/CinematicAssemblyTools/Source/CineAssemblyTools/Public/CineAssemblySchema.h b/Engine/Plugins/VirtualProduction/CinematicAssemblyTools/Source/CineAssemblyTools/Public/CineAssemblySchema.h
--- a/Engine/Plugins/VirtualProduction/CinematicAssemblyTools/Source/CineAssemblyTools/Public/CineAssemblySchema.h
+++ b/Engine/Plugins/VirtualProduction/CinematicAssemblyTools/Source/CineAssemblyTools/Public/CineAssemblySchema.h
@@ -96,7 +96,7 @@ public:
 	 * The default path to use when creating assemblies from this schema
 	 * When an assembly asset is created, this path will be appended to path where the asset would normally have been created.
 	 */
-	UPROPERTY(EditAnywhere, Category = "Default")
+	UPROPERTY()
 	FString DefaultAssemblyPath;
 
 	/** Restricts assemblies made from this schema to using this Schema when picking a Parent Assembly */
diff --git a/Engine/Plugins/VirtualProduction/CinematicAssemblyTools/Source/CineAssemblyToolsEditor/Private/CineAssemblySchemaCustomization.cpp b/Engine/Plugins/VirtualProduction/CinematicAssemblyTools/Source/CineAssemblyToolsEditor/Private/CineAssemblySchemaCustomization.cpp
new file mode 100644
--- /dev/null
+++ b/Engine/Plugins/VirtualProduction/CinematicAssemblyTools/Source/CineAssemblyToolsEditor/Private/CineAssemblySchemaCustomization.cpp
@@ -0,0 +1,116 @@
+// Copyright Epic Games, Inc. All Rights Reserved.
+
+#include "CineAssemblySchemaCustomization.h"
+
+#include "CineAssemblySchema.h"
+#include "DetailLayoutBuilder.h"
+#include "PropertyCustomizationHelpers.h"
+#include "Widgets/Input/SEditableTextBox.h"
+
+#define LOCTEXT_NAMESPACE "CineAssemblySchemaCustomization"
+
+TSharedRef<IDetailCustomization> FCineAssemblySchemaCustomization::MakeInstance()
+{
+	return MakeShared<FCineAssemblySchemaCustomization>();
+}
+
+void FCineAssemblySchemaCustomization::CustomizeDetails(IDetailLayoutBuilder& DetailBuilder)
+{
+	TArray<TWeakObjectPtr<UObject>> CustomizedObjects;
+	DetailBuilder.GetObjectsBeingCustomized(CustomizedObjects);
+
+	// Ensure that we are only customizing one object
+	if (CustomizedObjects.Num() != 1)
+	{
+		return;
+	}
+
+	CustomizedSchema = Cast<UCineAssemblySchema>(CustomizedObjects[0]);
+
+	IDetailCategoryBuilder& DefaultCategory = DetailBuilder.EditCategory("Default");
+	TSharedRef<IPropertyHandle> SchemaNamePropertyHandle = DetailBuilder.GetProperty(GET_MEMBER_NAME_CHECKED(UCineAssemblySchema, SchemaName));
+	TSharedRef<IPropertyHandle> DescriptionPropertyHandle = DetailBuilder.GetProperty(GET_MEMBER_NAME_CHECKED(UCineAssemblySchema, Description));
+	TSharedRef<IPropertyHandle> DefaultAssemblyNamePropertyHandle = DetailBuilder.GetProperty(GET_MEMBER_NAME_CHECKED(UCineAssemblySchema, DefaultAssemblyName));
+
+	// Add the properties back into the category in the correct order
+	IDetailPropertyRow& SchemaNamePropertyRow = DefaultCategory.AddProperty(SchemaNamePropertyHandle);
+	DefaultCategory.AddProperty(DescriptionPropertyHandle);
+	IDetailPropertyRow& DefaultAssemblyNamePropertyRow = DefaultCategory.AddProperty(DefaultAssemblyNamePropertyHandle);
+
+	// Customize the widget for the SchemaName property to add additional validation on the user input text
+	SchemaNamePropertyRow.CustomWidget()
+		.NameContent()
+		[
+			SchemaNamePropertyHandle->CreatePropertyNameWidget()
+		]
+		.ValueContent()
+		[
+			SNew(SEditableTextBox)
+				.Font(IDetailLayoutBuilder::GetDetailFont())
+				.Text_Lambda([this]() -> FText
+					{
+						return FText::FromString(CustomizedSchema->SchemaName);
+					})
+				.OnTextCommitted_Lambda([this](const FText& InText, ETextCommit::Type InCommitType)
+					{
+						CustomizedSchema->Modify();
+						CustomizedSchema->SchemaName = InText.ToString();
+					})
+				.OnVerifyTextChanged_Raw(this, &FCineAssemblySchemaCustomization::ValidateSchemaName)
+		];
+
+	// Customize the widget for the DefaultAssemblyName property to add additional validation on the user input text
+	DefaultAssemblyNamePropertyRow.CustomWidget()
+		.NameContent()
+		[
+			DefaultAssemblyNamePropertyHandle->CreatePropertyNameWidget()
+		]
+		.ValueContent()
+		[
+			SNew(SEditableTextBox)
+				.Font(IDetailLayoutBuilder::GetDetailFont())
+				.Text_Lambda([this]() -> FText
+					{
+						return FText::FromString(CustomizedSchema->DefaultAssemblyName);
+					})
+				.OnTextCommitted_Lambda([this](const FText& InText, ETextCommit::Type InCommitType)
+					{
+						CustomizedSchema->Modify();
+						CustomizedSchema->DefaultAssemblyName = InText.ToString();
+					})
+				.OnVerifyTextChanged_Raw(this, &FCineAssemblySchemaCustomization::ValidateDefaultAssemblyName)
+		];
+}
+
+bool FCineAssemblySchemaCustomization::ValidateSchemaName(const FText& InText, FText& OutErrorMessage) const
+{
+	return FName::IsValidXName(InText.ToString(), INVALID_OBJECTNAME_CHARACTERS INVALID_LONGPACKAGE_CHARACTERS, &OutErrorMessage);
+}
+
+bool FCineAssemblySchemaCustomization::ValidateDefaultAssemblyName(const FText& InText, FText& OutErrorMessage) const
+{
+	// Ensure that the name does not contain any characters that would be invalid for an asset name
+	// This matches the validation that would happen if the user was renaming an asset in the content browser
+	FString InvalidCharacters = INVALID_OBJECTNAME_CHARACTERS INVALID_LONGPACKAGE_CHARACTERS;
+
+	// These characters are actually valid, because we want to support naming tokens
+	InvalidCharacters = InvalidCharacters.Replace(TEXT("{}"), TEXT(""));
+	InvalidCharacters = InvalidCharacters.Replace(TEXT(":"), TEXT(""));
+
+	const FString PotentialName = InText.ToString();
+
+	if (!FName::IsValidXName(PotentialName, InvalidCharacters, &OutErrorMessage))
+	{
+		return false;
+	}
+
+	if (PotentialName.Contains(TEXT("{assembly}")))
+	{
+		OutErrorMessage = LOCTEXT("RecursiveAssemblyTokenError", "The default assembly name cannot use the {assembly} token");
+		return false;
+	}
+
+	return true;
+}
+
+#undef LOCTEXT_NAMESPACE
diff --git a/Engine/Plugins/VirtualProduction/CinematicAssemblyTools/Source/CineAssemblyToolsEditor/Private/CineAssemblySchemaCustomization.h b/Engine/Plugins/VirtualProduction/CinematicAssemblyTools/Source/CineAssemblyToolsEditor/Private/CineAssemblySchemaCustomization.h
new file mode 100644
--- /dev/null
+++ b/Engine/Plugins/VirtualProduction/CinematicAssemblyTools/Source/CineAssemblyToolsEditor/Private/CineAssemblySchemaCustomization.h
@@ -0,0 +1,30 @@
+// Copyright Epic Games, Inc. All Rights Reserved.
+
+#pragma once
+
+#include "IDetailCustomization.h"
+
+class UCineAssemblySchema;
+
+/**
+ * Detail customization for UCineAssemblySchema
+ */
+class FCineAssemblySchemaCustomization : public IDetailCustomization
+{
+public:
+	/** Makes a new instance of this detail layout class for a specific detail view requesting it */
+	static TSharedRef<IDetailCustomization> MakeInstance();
+
+	//~ Begin IDetailCustomization interface */
+	virtual void CustomizeDetails(IDetailLayoutBuilder& DetailBuilder) override;
+	//~ End IDetailCustomization interface */
+
+private:
+	/** Validates the user input text for the associated properties */
+	bool ValidateSchemaName(const FText& InText, FText& OutErrorMessage) const;
+	bool ValidateDefaultAssemblyName(const FText& InText, FText& OutErrorMessage) const;
+
+private:
+	/** The assembly schema being customized */
+	UCineAssemblySchema* CustomizedSchema;
+};
diff --git a/Engine/Plugins/VirtualProduction/CinematicAssemblyTools/Source/CineAssemblyToolsEditor/Private/CineAssemblyToolsEditorModule.cpp b/Engine/Plugins/VirtualProduction/CinematicAssemblyTools/Source/CineAssemblyToolsEditor/Private/CineAssemblyToolsEditorModule.cpp
--- a/Engine/Plugins/VirtualProduction/CinematicAssemblyTools/Source/CineAssemblyToolsEditor/Private/CineAssemblyToolsEditorModule.cpp
+++ b/Engine/Plugins/VirtualProduction/CinematicAssemblyTools/Source/CineAssemblyToolsEditor/Private/CineAssemblyToolsEditorModule.cpp
@@ -5,9 +5,10 @@
 #include "AssetRegistry/AssetRegistryModule.h"
 #include "AssetToolsModule.h"
 #include "CineAssembly.h"
-#include "CineAssemblySchema.h"
 #include "CineAssemblyCustomization.h"
 #include "CineAssemblyMetadataCustomization.h"
+#include "CineAssemblySchema.h"
+#include "CineAssemblySchemaCustomization.h"
 #include "CineAssemblyToolsStyle.h"
 #include "ContentBrowserMenuContexts.h"
 #include "Editor.h"
@@ -101,6 +102,11 @@ void FCineAssemblyToolsEditorModule::StartupModule()
 		FOnGetDetailCustomizationInstance::CreateStatic(&FCineAssemblyCustomization::MakeInstance)
 	);
 
+	PropertyEditorModule.RegisterCustomClassLayout(
+		UCineAssemblySchema::StaticClass()->GetFName(),
+		FOnGetDetailCustomizationInstance::CreateStatic(&FCineAssemblySchemaCustomization::MakeInstance)
+	);
+
 	PropertyEditorModule.RegisterCustomPropertyTypeLayout(
 		FAssemblyMetadataDesc::StaticStruct()->GetFName(),
 		FOnGetPropertyTypeCustomizationInstance::CreateStatic(&FCineAssemblyMetadataCustomization::MakeInstance)
@@ -169,6 +175,9 @@ void FCineAssemblyToolsEditorModule::ShutdownModule()
 
 	FPropertyEditorModule& PropertyModule = FModuleManager::GetModuleChecked<FPropertyEditorModule>("PropertyEditor");
 	PropertyModule.UnregisterCustomClassLayout(UProductionSettings::StaticClass()->GetFName());
+	PropertyModule.UnregisterCustomClassLayout(UCineAssembly::StaticClass()->GetFName());
+	PropertyModule.UnregisterCustomClassLayout(UCineAssemblySchema::StaticClass()->GetFName());
+	PropertyModule.UnregisterCustomPropertyTypeLayout(FAssemblyMetadataDesc::StaticStruct()->GetFName());
 
 	FGlobalTabmanager::Get()->UnregisterNomadTabSpawner(ProductionWizardTabName);
 
