ShotManagement: Fix Cine Assembly Schema edit windows to be dockable as nomad tabs.
Fix Cine Assembly Schema window to no longer require any work to be done OnWindowClosed(). Schema properties are all now updated as they are changed so that users can leave Schema tabs open and docked but have the changes to the asset take immediate effect.
Add GUID to Cine Assembly Schema assets.
Fix GUID assignment on Cine Assembly assets to properly handle duplication.
Resave the existing plugin content template schemas with unique IDs assigned now.
#rnx

#virtualized

diff --git a/Engine/Plugins/VirtualProduction/CinematicAssemblyTools/Source/CineAssemblyTools/Private/CineAssembly.cpp b/Engine/Plugins/VirtualProduction/CinematicAssemblyTools/Source/CineAssemblyTools/Private/CineAssembly.cpp
--- a/Engine/Plugins/VirtualProduction/CinematicAssemblyTools/Source/CineAssemblyTools/Private/CineAssembly.cpp
+++ b/Engine/Plugins/VirtualProduction/CinematicAssemblyTools/Source/CineAssemblyTools/Private/CineAssembly.cpp
@@ -26,6 +26,31 @@ const FName UCineAssembly::AssemblyGuidPropertyName = GET_MEMBER_NAME_CHECKED(UC
 UCineAssembly::UCineAssembly()
 {
 	MetadataJsonObject = MakeShared<FJsonObject>();
+}
+
+void UCineAssembly::PostInitProperties()
+{
+	Super::PostInitProperties();
+
+	if (!HasAnyFlags(RF_ClassDefaultObject | RF_ArchetypeObject | RF_NeedLoad | RF_WasLoaded) && !AssemblyGuid.IsValid())
+	{
+		AssemblyGuid = FGuid::NewGuid();
+	}
+}
+
+void UCineAssembly::PostDuplicate(bool bDuplicateForPIE)
+{
+	Super::PostDuplicate(bDuplicateForPIE);
+
+	if (!bDuplicateForPIE)
+	{
+		AssemblyGuid = FGuid::NewGuid();
+	}
+}
+
+void UCineAssembly::PostLoad()
+{
+	Super::PostLoad();
 
 	if (!AssemblyGuid.IsValid())
 	{
diff --git a/Engine/Plugins/VirtualProduction/CinematicAssemblyTools/Source/CineAssemblyTools/Private/CineAssemblySchema.cpp b/Engine/Plugins/VirtualProduction/CinematicAssemblyTools/Source/CineAssemblyTools/Private/CineAssemblySchema.cpp
--- a/Engine/Plugins/VirtualProduction/CinematicAssemblyTools/Source/CineAssemblyTools/Private/CineAssemblySchema.cpp
+++ b/Engine/Plugins/VirtualProduction/CinematicAssemblyTools/Source/CineAssemblyTools/Private/CineAssemblySchema.cpp
@@ -11,6 +11,7 @@
 #endif
 
 const FSoftObjectPath UCineAssemblySchema::DefaultThumbnailPath = FSoftObjectPath(TEXT("/CinematicAssemblyTools/DefaultSchemaThumbnail.DefaultSchemaThumbnail"));
+const FName UCineAssemblySchema::SchemaGuidPropertyName = GET_MEMBER_NAME_CHECKED(UCineAssemblySchema, SchemaGuid);
 
 UCineAssemblySchema::UCineAssemblySchema()
 {
@@ -22,6 +23,11 @@ UCineAssemblySchema::UCineAssemblySchema()
 	}
 }
 
+FGuid UCineAssemblySchema::GetSchemaGuid() const
+{
+	return SchemaGuid;
+}
+
 void UCineAssemblySchema::Serialize(FArchive& Ar)
 {
 	Super::Serialize(Ar);
@@ -31,10 +37,35 @@ void UCineAssemblySchema::Serialize(FArchive& Ar)
 	}
 }
 
+void UCineAssemblySchema::PostInitProperties()
+{
+	Super::PostInitProperties();
+
+	if (!HasAnyFlags(RF_ClassDefaultObject | RF_ArchetypeObject | RF_NeedLoad | RF_WasLoaded) && !SchemaGuid.IsValid())
+	{
+		SchemaGuid = FGuid::NewGuid();
+	}
+}
+
+void UCineAssemblySchema::PostDuplicate(bool bDuplicateForPIE)
+{
+	Super::PostDuplicate(bDuplicateForPIE);
+
+	if (!bDuplicateForPIE)
+	{
+		SchemaGuid = FGuid::NewGuid();
+	}
+}
+
 void UCineAssemblySchema::PostLoad()
 {
 	Super::PostLoad();
 
+	if (!SchemaGuid.IsValid())
+	{
+		SchemaGuid = FGuid::NewGuid();
+	}
+
 	UpdateThumbnailBrush();
 }
 
@@ -69,13 +100,14 @@ void UCineAssemblySchema::RenameAsset()
 
 	IAssetTools& AssetTools = FModuleManager::GetModuleChecked<FAssetToolsModule>("AssetTools").Get();
 
-	const FSoftObjectPath OldObjectPath = FSoftObjectPath(this);
+	const FString PackagePath = FPackageName::GetLongPackagePath(GetOutermost()->GetName());
 
-	const FTopLevelAssetPath NewAssetPath(*GetPathName(), *SchemaName);
-	const FSoftObjectPath NewObjectPath(NewAssetPath);
+	TArray<FAssetRenameData> AssetsAndNames;
+	const bool bSoftReferenceOnly = false;
+	const bool bAlsoRenameLocalizedVariants = true;
+	AssetsAndNames.Emplace(FAssetRenameData(this, PackagePath, SchemaName, bSoftReferenceOnly, bAlsoRenameLocalizedVariants));
 
-	FAssetRenameData RenameData = FAssetRenameData(OldObjectPath, NewObjectPath);
-	AssetTools.RenameAssets({ RenameData });
+	AssetTools.RenameAssetsWithDialog(AssetsAndNames);
 
 	bSupportsRename = false;
 }
diff --git a/Engine/Plugins/VirtualProduction/CinematicAssemblyTools/Source/CineAssemblyTools/Public/CineAssembly.h b/Engine/Plugins/VirtualProduction/CinematicAssemblyTools/Source/CineAssemblyTools/Public/CineAssembly.h
--- a/Engine/Plugins/VirtualProduction/CinematicAssemblyTools/Source/CineAssemblyTools/Public/CineAssembly.h
+++ b/Engine/Plugins/VirtualProduction/CinematicAssemblyTools/Source/CineAssemblyTools/Public/CineAssembly.h
@@ -31,6 +31,9 @@ public:
 
 	//~ Begin UObject Interface
 	virtual void GetAssetRegistryTags(FAssetRegistryTagsContext Context) const override;
+	virtual void PostInitProperties() override;
+	virtual void PostDuplicate(bool bDuplicateForPIE) override;
+	virtual void PostLoad() override;
 
 #if WITH_EDITOR
 	virtual void GetAssetRegistryTagMetadata(TMap<FName, FAssetRegistryTagMetadata>& OutMetadata) const override;
diff --git a/Engine/Plugins/VirtualProduction/CinematicAssemblyTools/Source/CineAssemblyTools/Public/CineAssemblySchema.h b/Engine/Plugins/VirtualProduction/CinematicAssemblyTools/Source/CineAssemblyTools/Public/CineAssemblySchema.h
--- a/Engine/Plugins/VirtualProduction/CinematicAssemblyTools/Source/CineAssemblyTools/Public/CineAssemblySchema.h
+++ b/Engine/Plugins/VirtualProduction/CinematicAssemblyTools/Source/CineAssemblyTools/Public/CineAssemblySchema.h
@@ -57,6 +57,8 @@ public:
 
 	//~ Begin UObject overrides
 	virtual void Serialize(FArchive& Ar) override;
+	virtual void PostInitProperties() override;
+	virtual void PostDuplicate(bool bDuplicateForPIE) override;
 	virtual void PostLoad() override;
 
 #if WITH_EDITOR
@@ -64,6 +66,9 @@ public:
 #endif
 	//~ End UObject overrides
 
+	/** Get the unique ID of this schema */
+	CINEASSEMBLYTOOLS_API FGuid GetSchemaGuid() const;
+
 	/** Returns the thumbnail brush used by this schema */
 	CINEASSEMBLYTOOLS_API const FSlateBrush* GetThumbnailBrush() const;
 
@@ -119,7 +124,13 @@ public:
 	UPROPERTY()
 	TArray<FString> FoldersToCreate;
 
+	static CINEASSEMBLYTOOLS_API const FName SchemaGuidPropertyName;
+
 private:
+	/** Unique ID for this schema, assigned at object creation */
+	UPROPERTY(AssetRegistrySearchable, meta = (IgnoreForMemberInitializationTest))
+	FGuid SchemaGuid;
+
 	/** Slate brush used to actually render the thumbnail image */
 	TSharedPtr<FSlateBrush> ThumbnailBrush;
 
diff --git a/Engine/Plugins/VirtualProduction/CinematicAssemblyTools/Source/CineAssemblyToolsEditor/Private/AssetDefinition_CineAssemblySchema.cpp b/Engine/Plugins/VirtualProduction/CinematicAssemblyTools/Source/CineAssemblyToolsEditor/Private/AssetDefinition_CineAssemblySchema.cpp
--- a/Engine/Plugins/VirtualProduction/CinematicAssemblyTools/Source/CineAssemblyToolsEditor/Private/AssetDefinition_CineAssemblySchema.cpp
+++ b/Engine/Plugins/VirtualProduction/CinematicAssemblyTools/Source/CineAssemblyToolsEditor/Private/AssetDefinition_CineAssemblySchema.cpp
@@ -3,7 +3,7 @@
 #include "AssetDefinition_CineAssemblySchema.h"
 
 #include "CineAssemblySchema.h"
-#include "Framework/Application/SlateApplication.h"
+#include "CineAssemblyToolsEditorModule.h"
 #include "UI/CineAssembly/SCineAssemblySchemaWindow.h"
 
 #define LOCTEXT_NAMESPACE "AssetTypeActions"
@@ -60,8 +60,8 @@ EAssetCommandResult UAssetDefinition_CineAssemblySchema::OpenAssets(const FAsset
 		const FAssetData& SchemaAssetData = OpenArgs.Assets[0];
 		if (UCineAssemblySchema* CineAssemblySchema = Cast<UCineAssemblySchema>(SchemaAssetData.GetAsset()))
 		{
-			TSharedRef<SCineAssemblySchemaWindow> CineAssemblySchemaWindow = SNew(SCineAssemblySchemaWindow, Cast<UCineAssemblySchema>(CineAssemblySchema), ESchemaConfigMode::Edit);
-			FSlateApplication::Get().AddWindow(CineAssemblySchemaWindow);
+			FCineAssemblyToolsEditorModule& CineAssemblyToolsEditorModule = FModuleManager::GetModuleChecked<FCineAssemblyToolsEditorModule>("CineAssemblyToolsEditor");
+			CineAssemblyToolsEditorModule.OpenSchemaForEdit(CineAssemblySchema);
 		}
 	}
 
diff --git a/Engine/Plugins/VirtualProduction/CinematicAssemblyTools/Source/CineAssemblyToolsEditor/Private/CineAssemblySchemaFactory.cpp b/Engine/Plugins/VirtualProduction/CinematicAssemblyTools/Source/CineAssemblyToolsEditor/Private/CineAssemblySchemaFactory.cpp
--- a/Engine/Plugins/VirtualProduction/CinematicAssemblyTools/Source/CineAssemblyToolsEditor/Private/CineAssemblySchemaFactory.cpp
+++ b/Engine/Plugins/VirtualProduction/CinematicAssemblyTools/Source/CineAssemblyToolsEditor/Private/CineAssemblySchemaFactory.cpp
@@ -27,8 +27,19 @@ bool UCineAssemblySchemaFactory::ConfigureProperties()
 	IContentBrowserSingleton& ContentBrowser = FModuleManager::LoadModuleChecked<FContentBrowserModule>("ContentBrowser").Get();
 	const FString CurrentPath = ContentBrowser.GetCurrentPath().GetInternalPathString();
 
-	TSharedRef<SCineAssemblySchemaWindow> CineAssemblySchemaWindow = SNew(SCineAssemblySchemaWindow, CurrentPath);
-	FSlateApplication::Get().AddWindow(CineAssemblySchemaWindow);
+	TSharedRef<SCineAssemblySchemaWindow> SchemaWidget = SNew(SCineAssemblySchemaWindow, CurrentPath);
+
+	const FVector2D DefaultWindowSize = FVector2D(1400, 750);
+
+	TSharedRef<SWindow> NewSchemaWindow = SNew(SWindow)
+		.Title(NSLOCTEXT("CineAssemblySchemaFactory", "WindowTitleCreateNew", "Create Assembly Schema"))
+		.SizingRule(ESizingRule::UserSized)
+		.ClientSize(DefaultWindowSize)
+		.SupportsMaximize(false)
+		.SupportsMinimize(false);
+
+	NewSchemaWindow->SetContent(SchemaWidget);
+	FSlateApplication::Get().AddWindow(NewSchemaWindow);
 
 	return false;
 }
diff --git a/Engine/Plugins/VirtualProduction/CinematicAssemblyTools/Source/CineAssemblyToolsEditor/Private/CineAssemblyToolsEditorModule.cpp b/Engine/Plugins/VirtualProduction/CinematicAssemblyTools/Source/CineAssemblyToolsEditor/Private/CineAssemblyToolsEditorModule.cpp
--- a/Engine/Plugins/VirtualProduction/CinematicAssemblyTools/Source/CineAssemblyToolsEditor/Private/CineAssemblyToolsEditorModule.cpp
+++ b/Engine/Plugins/VirtualProduction/CinematicAssemblyTools/Source/CineAssemblyToolsEditor/Private/CineAssemblyToolsEditorModule.cpp
@@ -1,6 +1,6 @@
 // Copyright Epic Games, Inc. All Rights Reserved.
 
-#include "Modules/ModuleManager.h"
+#include "CineAssemblyToolsEditorModule.h"
 
 #include "AssetRegistry/AssetRegistryModule.h"
 #include "AssetToolsModule.h"
@@ -12,9 +12,7 @@
 #include "CineAssemblyToolsStyle.h"
 #include "ContentBrowserMenuContexts.h"
 #include "Editor.h"
-#include "Framework/Application/SlateApplication.h"
 #include "Framework/Docking/TabManager.h"
-#include "Interfaces/IMainFrameModule.h"
 #include "Misc/ConfigCacheIni.h"
 #include "MoviePipelineQueue.h"
 #include "NamingTokensEngineSubsystem.h"
@@ -25,10 +23,9 @@
 #include "ToolMenus.h"
 #include "UI/SProductionWizard.h"
 #include "UI/CineAssembly/SCineAssemblyConfigWindow.h"
-#include "Widgets/Docking/SDockTab.h"
+#include "UI/CineAssembly/SCineAssemblySchemaWindow.h"
 #include "WorkspaceMenuStructure.h"
 #include "WorkspaceMenuStructureModule.h"
-#include "TakeRecorder/CineAssemblyTakeRecorderIntegration.h"
 
 #define LOCTEXT_NAMESPACE "CineAssemblyToolsEditorModule"
 
@@ -36,55 +33,7 @@ DEFINE_LOG_CATEGORY_STATIC(LogCineAssemblyToolsEditorModule, Log, All)
 
 static const FName ProductionWizardTabName("ProductionWizard");
 
-class FCineAssemblyToolsEditorModule : public IModuleInterface
-{
-public:
-
-	//~ Begin IModuleInterface interface
-	virtual void StartupModule() override;
-	virtual void ShutdownModule() override;
-	//~ End IModuleInterface interface
-
-private:
-	void OnPostEngineInit();
-
-	/** Spawns a nomad tab for the Production Wizard tool */
-	TSharedRef<SDockTab> SpawnProductionWizard(const FSpawnTabArgs& SpawnTabArgs);
-
-	/** Spawns a nomad tab for an Assembly asset to edit its properties */
-	TSharedRef<SDockTab> SpawnAssemblyTab(const FSpawnTabArgs& SpawnTabArgs, UCineAssembly* Assembly);
-	TSharedRef<SDockTab> SpawnAssemblyTab(const FSpawnTabArgs& SpawnTabArgs, FGuid AssemblyGuid);
-	TSharedRef<SDockTab> SpawnAssemblyTab(const FSpawnTabArgs& SpawnTabArgs, TSharedPtr<SCineAssemblyEditWidget> Widget);
-
-	/** Add context menu actions */
-	void RegisterMenus();
-
-	/** Called whenever a job has its sequence set. */
-	void OnSequenceSet(UMoviePipelineExecutorJob* Job, ULevelSequence* Sequence);
-
-	/** Opens a window to edit the properties of the input assembly */
-	void OpenAssemblyForEdit(UCineAssembly* Assembly);
-
-	/** Returns the next available tab name in the map of managed tabs. If there are no available tabs, returns false */
-	bool TryGetNextAssemblyTab(FTabId& OutTabId);
-
-	/** Callback when an asset is deleted which allows us to check if it was a Cine Assembly and then close the tab for that Assembly (if one was open) */
-	void OnAssetDeleted(UObject* Object);
-
-private:
-	/** Manages all integration of Cinematic Assembly Tools with Take Recorder. */
-	TUniquePtr<FCineAssemblyTakeRecorderIntegration> TakeRecorderIntegration;
-
-	/** Map of TabIds to AssemblyIDs to track which assets are open in which tabs */
-	TMap<FTabId, FGuid> ManagedTabs;
-	const int32 MaxAssemblyTabs = 16;
-
-	static const FString OpenAssemblySection;
-
-	FDelegateHandle OnSequenceSetHandle;
-};
-
-const FString FCineAssemblyToolsEditorModule::OpenAssemblySection = TEXT("OpenAssemblyTabs");
+const FString FCineAssemblyToolsEditorModule::OpenTabSection = TEXT("CinematicAssemblyTools_OpenTabs");
 
 void FCineAssemblyToolsEditorModule::StartupModule()
 {
@@ -125,11 +74,15 @@ void FCineAssemblyToolsEditorModule::StartupModule()
 
 	OnSequenceSetHandle = UMoviePipelineExecutorJob::OnSequenceSetEvent().AddRaw(this, &FCineAssemblyToolsEditorModule::OnSequenceSet);
 
-	// Initialize the map of managed tabs
-	for (int32 TabIndex = 0; TabIndex < MaxAssemblyTabs; ++TabIndex)
+	// Initialize the maps of managed tabs
+	constexpr int32 MaxAssetTabs = 16;
+	for (int32 TabIndex = 0; TabIndex < MaxAssetTabs; ++TabIndex)
 	{
 		const FString AssemblyTabName = FString::Printf(TEXT("CineAssemblyTab%d"), TabIndex);
-		ManagedTabs.Add(FTabId(*AssemblyTabName), FGuid());
+		ManagedAssemblyTabs.Add(FTabId(*AssemblyTabName), FGuid());
+
+		const FString SchemaTabName = FString::Printf(TEXT("CineAssemblySchemaTab%d"), TabIndex);
+		ManagedSchemaTabs.Add(FTabId(*SchemaTabName), FGuid());
 	}
 }
 
@@ -149,17 +102,30 @@ void FCineAssemblyToolsEditorModule::OnPostEngineInit()
 		}));
 
 	// Load any previously opened assembly tabs and register new nomad tab spawners for them so they can be properly restored in the layout.
-	for (TPair<FTabId, FGuid>& TabPair : ManagedTabs)
+	for (TPair<FTabId, FGuid>& TabPair : ManagedAssemblyTabs)
 	{
 		const FName TabName = TabPair.Key.TabType;
+		const FGuid AssetID = FindTabAssetInConfig(TabName);
 
-		FString GuidString;
-		if (GConfig->GetString(*OpenAssemblySection, *TabName.ToString(), GuidString, GEditorPerProjectIni))
+		if (AssetID.IsValid())
 		{
-			const FGuid AssemblyID = FGuid(GuidString);
-			TabPair.Value = AssemblyID;
+			TabPair.Value = AssetID;
+
+			FGlobalTabmanager::Get()->RegisterNomadTabSpawner(TabName, FOnSpawnTab::CreateRaw(this, &FCineAssemblyToolsEditorModule::SpawnAssemblyTab, AssetID))
+				.SetMenuType(ETabSpawnerMenuType::Hidden);
+		}
+	}
+
+	// Load any previously opened schema tabs and register new nomad tab spawners for them so they can be properly restored in the layout.
+	for (TPair<FTabId, FGuid>& TabPair : ManagedSchemaTabs)
+	{
+		const FName TabName = TabPair.Key.TabType;
+		const FGuid AssetID = FindTabAssetInConfig(TabName);
 
-			FGlobalTabmanager::Get()->RegisterNomadTabSpawner(TabName, FOnSpawnTab::CreateRaw(this, &FCineAssemblyToolsEditorModule::SpawnAssemblyTab, AssemblyID))
+		if (AssetID.IsValid())
+		{
+			TabPair.Value = AssetID;
+			FGlobalTabmanager::Get()->RegisterNomadTabSpawner(TabName, FOnSpawnTab::CreateRaw(this, &FCineAssemblyToolsEditorModule::SpawnSchemaTab, AssetID))
 				.SetMenuType(ETabSpawnerMenuType::Hidden);
 		}
 	}
@@ -193,22 +159,9 @@ void FCineAssemblyToolsEditorModule::ShutdownModule()
 		AssetRegistryModule->Get().OnInMemoryAssetDeleted().RemoveAll(this);
 	}
 
-	// Write out each managed tab that is associated with a valid Assembly to a config file so that we can restore them the next time the editor launches
-	GConfig->EmptySection(*OpenAssemblySection, GEditorPerProjectIni);
-	for (const TPair<FTabId, FGuid>& TabPair : ManagedTabs)
-	{
-		const FGuid AssemblyID = TabPair.Value;
-		if (AssemblyID.IsValid())
-		{
-			GConfig->AddToSection(*OpenAssemblySection, TabPair.Key.TabType, AssemblyID.ToString(), GEditorPerProjectIni);
-
-			if (TSharedPtr<FGlobalTabmanager> GlobalTabManager = FGlobalTabmanager::Get())
-			{
-				GlobalTabManager->UnregisterNomadTabSpawner(TabPair.Key.TabType);
-			}
-		}
-	}
-	GConfig->Flush(false, GEditorPerProjectIni);
+	GConfig->EmptySection(*OpenTabSection, GEditorPerProjectIni);
+	SaveOpenTabs(ManagedAssemblyTabs);
+	SaveOpenTabs(ManagedSchemaTabs);
 }
 
 TSharedRef<SDockTab> FCineAssemblyToolsEditorModule::SpawnProductionWizard(const FSpawnTabArgs& SpawnTabArgs)
@@ -244,17 +197,56 @@ TSharedRef<SDockTab> FCineAssemblyToolsEditorModule::SpawnAssemblyTab(const FSpa
 			Widget.ToSharedRef()
 		];
 
+	AssemblyTab->SetTabIcon(FSlateIcon(FCineAssemblyToolsStyle::StyleName, "Icons.Assembly").GetIcon());
+
 	// Unregister the spawner and reset the map entry when the tab closes
 	const FTabId TabId = SpawnTabArgs.GetTabId();
 	AssemblyTab->SetOnTabClosed(SDockTab::FOnTabClosedCallback::CreateLambda([this, TabId](TSharedRef<SDockTab> Tab)
 		{
 			FGlobalTabmanager::Get()->UnregisterNomadTabSpawner(*TabId.ToString());
-			ManagedTabs[TabId].Invalidate();
+			ManagedAssemblyTabs[TabId].Invalidate();
 		}));
 
 	return AssemblyTab;
 }
 
+TSharedRef<SDockTab> FCineAssemblyToolsEditorModule::SpawnSchemaTab(const FSpawnTabArgs& SpawnTabArgs, FGuid SchemaGuid)
+{
+	TSharedPtr<SCineAssemblySchemaWindow> Widget = SNew(SCineAssemblySchemaWindow, SchemaGuid);
+	return SpawnSchemaTab(SpawnTabArgs, Widget);
+}
+
+TSharedRef<SDockTab> FCineAssemblyToolsEditorModule::SpawnSchemaTab(const FSpawnTabArgs& SpawnTabArgs, UCineAssemblySchema* Schema)
+{
+	TSharedPtr<SCineAssemblySchemaWindow> Widget = SNew(SCineAssemblySchemaWindow, Schema);
+	return SpawnSchemaTab(SpawnTabArgs, Widget);
+}
+
+TSharedRef<SDockTab> FCineAssemblyToolsEditorModule::SpawnSchemaTab(const FSpawnTabArgs& SpawnTabArgs, TSharedPtr<SCineAssemblySchemaWindow> Widget)
+{
+	TSharedRef<SDockTab> SchemaTab = SNew(SDockTab)
+		.TabRole(ETabRole::NomadTab)
+		.Label_Lambda([Widget]() -> FText
+			{
+				return FText::FromString(Widget->GetSchemaName());
+			})
+		[
+			Widget.ToSharedRef()
+		];
+
+	SchemaTab->SetTabIcon(FSlateIcon(FCineAssemblyToolsStyle::StyleName, "Icons.Schema").GetIcon());
+
+	// Unregister the spawner and reset the map entry when the tab closes
+	const FTabId TabId = SpawnTabArgs.GetTabId();
+	SchemaTab->SetOnTabClosed(SDockTab::FOnTabClosedCallback::CreateLambda([this, TabId](TSharedRef<SDockTab> Tab)
+		{
+			FGlobalTabmanager::Get()->UnregisterNomadTabSpawner(*TabId.ToString());
+			ManagedSchemaTabs[TabId].Invalidate();
+		}));
+
+	return SchemaTab;
+}
+
 void FCineAssemblyToolsEditorModule::RegisterMenus()
 {
 	UToolMenus* ToolMenus = UToolMenus::Get();
@@ -307,7 +299,7 @@ void FCineAssemblyToolsEditorModule::OpenAssemblyForEdit(UCineAssembly* Assembly
 {
 	// Check if the assembly is already open in one of the managed tabs, and if so, bring focus to the existing tab
 	const FGuid AssemblyID = Assembly->GetAssemblyGuid();
-	const FTabId* ExistingTab = ManagedTabs.FindKey(AssemblyID);
+	const FTabId* ExistingTab = ManagedAssemblyTabs.FindKey(AssemblyID);
 
 	if (ExistingTab)
 	{
@@ -317,8 +309,9 @@ void FCineAssemblyToolsEditorModule::OpenAssemblyForEdit(UCineAssembly* Assembly
 
 	// Try to get the TabID for the next available tab in the map.
 	FTabId NewTabId;
-	if (!TryGetNextAssemblyTab(NewTabId))
+	if (!TryGetNextTab(ManagedAssemblyTabs, NewTabId))
 	{
+		UE_LOG(LogCineAssemblyToolsEditorModule, Warning, TEXT("The Cine Assembly tab could not be opened because the maximum number of Assembly tabs are already open."));
 		return;
 	}
 
@@ -328,13 +321,42 @@ void FCineAssemblyToolsEditorModule::OpenAssemblyForEdit(UCineAssembly* Assembly
 
 	FGlobalTabmanager::Get()->TryInvokeTab(NewTabId);
 
-	ManagedTabs[NewTabId] = AssemblyID;
+	ManagedAssemblyTabs[NewTabId] = AssemblyID;
+}
+
+void FCineAssemblyToolsEditorModule::OpenSchemaForEdit(UCineAssemblySchema* Schema)
+{
+	// Check if the assembly is already open in one of the managed tabs, and if so, bring focus to the existing tab
+	const FGuid SchemaID = Schema->GetSchemaGuid();
+	const FTabId* ExistingTab = ManagedSchemaTabs.FindKey(SchemaID);
+
+	if (ExistingTab)
+	{
+		FGlobalTabmanager::Get()->TryInvokeTab(*ExistingTab);
+		return;
+	}
+
+	// Try to get the TabID for the next available tab in the map.
+	FTabId NewTabId;
+	if (!TryGetNextTab(ManagedSchemaTabs, NewTabId))
+	{
+		UE_LOG(LogCineAssemblyToolsEditorModule, Warning, TEXT("The Cine Assembly Schema tab could not be opened because the maximum number of Schema tabs are already open."));
+		return;
+	}
+
+	// Register a new spawner and invoke a new tab to edit the properties of the assembly asset
+	FTabSpawnerEntry& Tmp = FGlobalTabmanager::Get()->RegisterNomadTabSpawner(NewTabId.TabType, FOnSpawnTab::CreateRaw(this, &FCineAssemblyToolsEditorModule::SpawnSchemaTab, Schema))
+		.SetMenuType(ETabSpawnerMenuType::Hidden);
+
+	FGlobalTabmanager::Get()->TryInvokeTab(NewTabId);
+
+	ManagedSchemaTabs[NewTabId] = SchemaID;
 }
 
-bool FCineAssemblyToolsEditorModule::TryGetNextAssemblyTab(FTabId& OutTabId)
+bool FCineAssemblyToolsEditorModule::TryGetNextTab(const TMap<FTabId, FGuid>& TabMap, FTabId& OutTabId)
 {
-	// Find the first tab in the map that is not currently mapped to a valid assembly
-	for (const TPair<FTabId, FGuid>& TabPair : ManagedTabs)
+	// Find the first tab in the map that is not currently mapped to a valid asset
+	for (const TPair<FTabId, FGuid>& TabPair : TabMap)
 	{
 		if (!TabPair.Value.IsValid())
 		{
@@ -343,26 +365,70 @@ bool FCineAssemblyToolsEditorModule::TryGetNextAssemblyTab(FTabId& OutTabId)
 		}
 	}
 
-	// If all tabs are currently mapped, we return false so that no new tab is opened
-	UE_LOG(LogCineAssemblyToolsEditorModule, Warning, TEXT("The Cine Assembly tab could not be opened because the maximum number of Assembly tabs are already open."));
 	return false;
 }
 
-void FCineAssemblyToolsEditorModule::OnAssetDeleted(UObject* Object)
+FGuid FCineAssemblyToolsEditorModule::FindTabAssetInConfig(FName TabName)
 {
-	if (const UCineAssembly* Assembly = Cast<UCineAssembly>(Object))
+	FString GuidString;
+	if (GConfig->GetString(*OpenTabSection, *TabName.ToString(), GuidString, GEditorPerProjectIni))
 	{
-		// Find the entry in the tab map matching the ID of the Assembly being deleted (if there is one)
-		if (const FTabId* OpenTabId = ManagedTabs.FindKey(Assembly->GetAssemblyGuid()))
+		return FGuid(GuidString);
+	}
+	return FGuid();
+}
+
+void FCineAssemblyToolsEditorModule::SaveOpenTabs(const TMap<FTabId, FGuid>& TabMap)
+{
+	// Write out each managed tab that is associated with a valid asset to a config file so that we can restore them the next time the editor launches
+	for (const TPair<FTabId, FGuid>& TabPair : TabMap)
+	{
+		const FName TabName = TabPair.Key.TabType;
+		const FGuid AssetID = TabPair.Value;
+		if (AssetID.IsValid())
 		{
-			// Get the open tab that is already open for this assembly and then close it. 
-			constexpr bool bInvokeAsInactive = true;
-			if (TSharedPtr<SDockTab> AssemblyTab = FGlobalTabmanager::Get()->TryInvokeTab(*OpenTabId, bInvokeAsInactive))
+			GConfig->AddToSection(*OpenTabSection, TabName, AssetID.ToString(), GEditorPerProjectIni);
+
+			if (TSharedPtr<FGlobalTabmanager> GlobalTabManager = FGlobalTabmanager::Get())
 			{
-				AssemblyTab->RequestCloseTab();
+				GlobalTabManager->UnregisterNomadTabSpawner(TabName);
 			}
+		}
+	}
+}
 
-			return;
+void FCineAssemblyToolsEditorModule::OnAssetDeleted(UObject* Object)
+{
+	// We only care about CineAssembly and CineAssemblySchema assets
+	if (!(Object->IsA(UCineAssembly::StaticClass()) || Object->IsA(UCineAssemblySchema::StaticClass())))
+	{
+		return;
+	}
+
+	// Find the tab matching the ID of the asset being deleted (if there is one)
+	FTabId AssetTabId;
+	if (const UCineAssembly* Assembly = Cast<UCineAssembly>(Object))
+	{
+		if (const FTabId* OpenTabId = ManagedAssemblyTabs.FindKey(Assembly->GetAssemblyGuid()))
+		{
+			AssetTabId = *OpenTabId;
+		}
+	}
+	else if (const UCineAssemblySchema* Schema = Cast<UCineAssemblySchema>(Object))
+	{
+		if (const FTabId* OpenTabId = ManagedSchemaTabs.FindKey(Schema->GetSchemaGuid()))
+		{
+			AssetTabId = *OpenTabId;
+		}
+	}
+
+	// Get the already opened tab and close it
+	if (!AssetTabId.TabType.IsNone())
+	{
+		constexpr bool bInvokeAsInactive = true;
+		if (TSharedPtr<SDockTab> AssetTab = FGlobalTabmanager::Get()->TryInvokeTab(AssetTabId, bInvokeAsInactive))
+		{
+			AssetTab->RequestCloseTab();
 		}
 	}
 }
diff --git a/Engine/Plugins/VirtualProduction/CinematicAssemblyTools/Source/CineAssemblyToolsEditor/Private/CineAssemblyToolsEditorModule.h b/Engine/Plugins/VirtualProduction/CinematicAssemblyTools/Source/CineAssemblyToolsEditor/Private/CineAssemblyToolsEditorModule.h
new file mode 100644
--- /dev/null
+++ b/Engine/Plugins/VirtualProduction/CinematicAssemblyTools/Source/CineAssemblyToolsEditor/Private/CineAssemblyToolsEditorModule.h
@@ -0,0 +1,77 @@
+// Copyright Epic Games, Inc. All Rights Reserved.
+
+#pragma once
+
+#include "Modules/ModuleManager.h"
+
+#include "TakeRecorder/CineAssemblyTakeRecorderIntegration.h"
+#include "Widgets/Docking/SDockTab.h"
+
+class UCineAssembly;
+class UCineAssemblySchema;
+class SCineAssemblyEditWidget;
+class SCineAssemblySchemaWindow;
+class UMoviePipelineExecutorJob;
+class ULevelSequence;
+
+class FCineAssemblyToolsEditorModule : public IModuleInterface
+{
+public:
+
+	//~ Begin IModuleInterface interface
+	virtual void StartupModule() override;
+	virtual void ShutdownModule() override;
+	//~ End IModuleInterface interface
+
+	/** Opens a window to edit the properties of the input assembly schema */
+	void OpenSchemaForEdit(UCineAssemblySchema* Schema);
+
+private:
+	void OnPostEngineInit();
+
+	/** Spawns a nomad tab for the Production Wizard tool */
+	TSharedRef<SDockTab> SpawnProductionWizard(const FSpawnTabArgs& SpawnTabArgs);
+
+	/** Spawns a nomad tab for an Assembly asset to edit its properties */
+	TSharedRef<SDockTab> SpawnAssemblyTab(const FSpawnTabArgs& SpawnTabArgs, UCineAssembly* Assembly);
+	TSharedRef<SDockTab> SpawnAssemblyTab(const FSpawnTabArgs& SpawnTabArgs, FGuid AssemblyGuid);
+	TSharedRef<SDockTab> SpawnAssemblyTab(const FSpawnTabArgs& SpawnTabArgs, TSharedPtr<SCineAssemblyEditWidget> Widget);
+
+	/** Spawns a nomad tab for an Schema asset to edit its properties */
+	TSharedRef<SDockTab> SpawnSchemaTab(const FSpawnTabArgs& SpawnTabArgs, UCineAssemblySchema* Schema);
+	TSharedRef<SDockTab> SpawnSchemaTab(const FSpawnTabArgs& SpawnTabArgs, FGuid SchemaGuid);
+	TSharedRef<SDockTab> SpawnSchemaTab(const FSpawnTabArgs& SpawnTabArgs, TSharedPtr<SCineAssemblySchemaWindow> Widget);
+
+	/** Add context menu actions */
+	void RegisterMenus();
+
+	/** Called whenever a job has its sequence set. */
+	void OnSequenceSet(UMoviePipelineExecutorJob* Job, ULevelSequence* Sequence);
+
+	/** Opens a window to edit the properties of the input assembly */
+	void OpenAssemblyForEdit(UCineAssembly* Assembly);
+
+	/** Returns the next available tab name in the map of managed tabs. If there are no available tabs, returns false */
+	bool TryGetNextTab(const TMap<FTabId, FGuid>& TabMap, FTabId& OutTabId);
+
+	/** Callback when an asset is deleted which allows us to check if it was a Cine Assembly and then close the tab for that Assembly (if one was open) */
+	void OnAssetDeleted(UObject* Object);
+
+	/** Look up the tab name in the open tabs config section and return the asset ID associated with that tab (if there is one) */
+	FGuid FindTabAssetInConfig(FName TabName);
+
+	/** Write out all of the tabs in the input tab map that are currently open (i.e. have a valid asset ID) */
+	void SaveOpenTabs(const TMap<FTabId, FGuid>& TabMap);
+
+private:
+	/** Manages all integration of Cinematic Assembly Tools with Take Recorder. */
+	TUniquePtr<FCineAssemblyTakeRecorderIntegration> TakeRecorderIntegration;
+
+	/** Map of TabIds to Asset IDs to track which assets are open in which tabs */
+	TMap<FTabId, FGuid> ManagedAssemblyTabs;
+	TMap<FTabId, FGuid> ManagedSchemaTabs;
+
+	static const FString OpenTabSection;
+
+	FDelegateHandle OnSequenceSetHandle;
+};
diff --git a/Engine/Plugins/VirtualProduction/CinematicAssemblyTools/Source/CineAssemblyToolsEditor/Private/CineAssemblyToolsStyle.cpp b/Engine/Plugins/VirtualProduction/CinematicAssemblyTools/Source/CineAssemblyToolsEditor/Private/CineAssemblyToolsStyle.cpp
--- a/Engine/Plugins/VirtualProduction/CinematicAssemblyTools/Source/CineAssemblyToolsEditor/Private/CineAssemblyToolsStyle.cpp
+++ b/Engine/Plugins/VirtualProduction/CinematicAssemblyTools/Source/CineAssemblyToolsEditor/Private/CineAssemblyToolsStyle.cpp
@@ -30,6 +30,9 @@ FCineAssemblyToolsStyle::FCineAssemblyToolsStyle()
 
 	Set("Thumbnails.Schema", new FSlateVectorImageBrush(RootToContentDir(TEXT("Slate/LevelSequenceActor_64.svg")), Icon64x64));
 
+	Set("Icons.Assembly", new FSlateVectorImageBrush(RootToContentDir(TEXT("Slate/Assembly.svg")), Icon16x16));
+	Set("Icons.Schema", new FSlateVectorImageBrush(RootToContentDir(TEXT("Slate/Schema.svg")), Icon16x16));
+
 	Set("Icons.RevisionControl", new FSlateVectorImageBrush(RootToContentDir(TEXT("Slate/RevisionControl.svg")), Icon16x16));
 	Set("Icons.Productions", new FSlateVectorImageBrush(RootToContentDir(TEXT("Slate/Camera.svg")), Icon16x16));
 	Set("Icons.Sequencer", new FSlateVectorImageBrush(RootToContentDir(TEXT("Slate/Sequencer_16.svg")), Icon16x16));
@@ -40,7 +43,6 @@ FCineAssemblyToolsStyle::FCineAssemblyToolsStyle()
 	Set("Icons.Initialize", new FSlateVectorImageBrush(RootToContentDir(TEXT("Slate/Initialize_16.svg")), Icon16x16));
 	Set("Icons.Animation", new FSlateVectorImageBrush(RootToContentDir(TEXT("Slate/Animation_16.svg")), Icon16x16));
 	Set("Icons.Notes", new FSlateVectorImageBrush(RootToContentDir(TEXT("Slate/Note.svg")), Icon16x16));
-	Set("Icons.Shot", new FSlateVectorImageBrush(RootToContentDir(TEXT("Slate/Shot_16.svg")), Icon16x16));
 	Set("Icons.DataAsset", new FSlateVectorImageBrush(RootToContentDir(TEXT("Slate/DataAsset.svg")), Icon16x16));
 	Set("Icons.DragHandle", new FSlateVectorImageBrush(RootToContentDir(TEXT("Slate/drag-handle.svg")), Icon16x16));
 
diff --git a/Engine/Plugins/VirtualProduction/CinematicAssemblyTools/Source/CineAssemblyToolsEditor/Private/UI/CineAssembly/SCineAssemblyConfigWindow.cpp b/Engine/Plugins/VirtualProduction/CinematicAssemblyTools/Source/CineAssemblyToolsEditor/Private/UI/CineAssembly/SCineAssemblyConfigWindow.cpp
--- a/Engine/Plugins/VirtualProduction/CinematicAssemblyTools/Source/CineAssemblyToolsEditor/Private/UI/CineAssembly/SCineAssemblyConfigWindow.cpp
+++ b/Engine/Plugins/VirtualProduction/CinematicAssemblyTools/Source/CineAssemblyToolsEditor/Private/UI/CineAssembly/SCineAssemblyConfigWindow.cpp
@@ -56,8 +56,6 @@ void SCineAssemblyConfigWindow::Construct(const FArguments& InArgs, const FStrin
 	UWorld* CurrentWorld = GEditor->GetEditorWorldContext().World();
 	CineAssemblyToConfigure->Level = FSoftObjectPath(CurrentWorld);
 
-	GetOnWindowClosedEvent().AddSP(this, &SCineAssemblyConfigWindow::OnAssemblyWindowClosed);
-
 	const FVector2D DefaultWindowSize = FVector2D(1400, 750);
 
 	SWindow::Construct(SWindow::FArguments()
@@ -108,20 +106,23 @@ void SCineAssemblyConfigWindow::Construct(const FArguments& InArgs, const FStrin
 		]);
 }
 
-void SCineAssemblyConfigWindow::OnAssemblyWindowClosed(const TSharedRef<SWindow>& InWindow)
+SCineAssemblyConfigWindow::~SCineAssemblyConfigWindow()
 {
 	// Save the UI config settings for whether to display engine/plugin content
-	UContentBrowserSettings* ContentBrowserSettings = GetMutableDefault<UContentBrowserSettings>();
-
-	const bool bShowEngineContent = ContentBrowserSettings->GetDisplayEngineFolder();
-	const bool bShowPluginContent = ContentBrowserSettings->GetDisplayPluginFolders();
+	if (UContentBrowserSettings* ContentBrowserSettings = GetMutableDefault<UContentBrowserSettings>())
+	{
+		const bool bShowEngineContent = ContentBrowserSettings->GetDisplayEngineFolder();
+		const bool bShowPluginContent = ContentBrowserSettings->GetDisplayPluginFolders();
 
-	GConfig->SetBool(TEXT("NewCineAssemblyUI"), TEXT("bShowEngineContent"), bShowEngineContent, GEditorPerProjectIni);
-	GConfig->SetBool(TEXT("NewCineAssemblyUI"), TEXT("bShowPluginContent"), bShowPluginContent, GEditorPerProjectIni);
-	GConfig->Flush(false, GEditorPerProjectIni);
+		if (GConfig)
+		{
+			GConfig->SetBool(TEXT("NewCineAssemblyUI"), TEXT("bShowEngineContent"), bShowEngineContent, GEditorPerProjectIni);
+			GConfig->SetBool(TEXT("NewCineAssemblyUI"), TEXT("bShowPluginContent"), bShowPluginContent, GEditorPerProjectIni);
+		}
 
-	ContentBrowserSettings->SetDisplayEngineFolder(bShowEngineContentCached);
-	ContentBrowserSettings->SetDisplayPluginFolders(bShowPluginContentCached);
+		ContentBrowserSettings->SetDisplayEngineFolder(bShowEngineContentCached);
+		ContentBrowserSettings->SetDisplayPluginFolders(bShowPluginContentCached);
+	}
 }
 
 TSharedRef<SWidget> SCineAssemblyConfigWindow::MakeCineTemplatePanel()
diff --git a/Engine/Plugins/VirtualProduction/CinematicAssemblyTools/Source/CineAssemblyToolsEditor/Private/UI/CineAssembly/SCineAssemblyConfigWindow.h b/Engine/Plugins/VirtualProduction/CinematicAssemblyTools/Source/CineAssemblyToolsEditor/Private/UI/CineAssembly/SCineAssemblyConfigWindow.h
--- a/Engine/Plugins/VirtualProduction/CinematicAssemblyTools/Source/CineAssemblyToolsEditor/Private/UI/CineAssembly/SCineAssemblyConfigWindow.h
+++ b/Engine/Plugins/VirtualProduction/CinematicAssemblyTools/Source/CineAssemblyToolsEditor/Private/UI/CineAssembly/SCineAssemblyConfigWindow.h
@@ -41,6 +41,7 @@ class SCineAssemblyConfigWindow : public SWindow
 {
 public:
 	SCineAssemblyConfigWindow() = default;
+	~SCineAssemblyConfigWindow();
 
 	SLATE_BEGIN_ARGS(SCineAssemblyConfigWindow) {}
 	SLATE_END_ARGS()
@@ -72,9 +73,6 @@ private:
 	/** Closes the window and indicates that no assets should be created by the asset factory */
 	FReply OnCancelClicked();
 
-	/** Callback when this window is about to be closed */
-	void OnAssemblyWindowClosed(const TSharedRef<SWindow>& InWindow);
-
 	/** Updates the UI and CineAssembly properties based on the selected schema */
 	void OnSchemaSelected(const FAssetData& AssetData);
 
diff --git a/Engine/Plugins/VirtualProduction/CinematicAssemblyTools/Source/CineAssemblyToolsEditor/Private/UI/CineAssembly/SCineAssemblySchemaWindow.cpp b/Engine/Plugins/VirtualProduction/CinematicAssemblyTools/Source/CineAssemblyToolsEditor/Private/UI/CineAssembly/SCineAssemblySchemaWindow.cpp
--- a/Engine/Plugins/VirtualProduction/CinematicAssemblyTools/Source/CineAssemblyToolsEditor/Private/UI/CineAssembly/SCineAssemblySchemaWindow.cpp
+++ b/Engine/Plugins/VirtualProduction/CinematicAssemblyTools/Source/CineAssemblyToolsEditor/Private/UI/CineAssembly/SCineAssemblySchemaWindow.cpp
@@ -38,27 +38,53 @@ void SCineAssemblySchemaWindow::Construct(const FArguments& InArgs, const FStrin
 {
 	// Create a new transient CineAssemblySchema to configure in the UI.
 	// If the configuration is successful, this will turn into the persistent object created by the factory.
-	UCineAssemblySchema* NewSchema = NewObject<UCineAssemblySchema>(GetTransientPackage(), NAME_None, RF_Transient);
+	SchemaToConfigure = TStrongObjectPtr<UCineAssemblySchema>(NewObject<UCineAssemblySchema>(GetTransientPackage(), NAME_None, RF_Transient));
+	Mode = ESchemaConfigMode::CreateNew;
+
 	CreateAssetPath = InCreateAssetPath;
 
-	Construct(InArgs, NewSchema, ESchemaConfigMode::CreateNew);
+	ChildSlot [ BuildUI() ];
 }
 
-void SCineAssemblySchemaWindow::Construct(const FArguments& InArgs, UCineAssemblySchema* InSchema, ESchemaConfigMode InConfigMode)
+void SCineAssemblySchemaWindow::Construct(const FArguments& InArgs, UCineAssemblySchema* InSchema)
 {
 	SchemaToConfigure = TStrongObjectPtr<UCineAssemblySchema>(InSchema);
-	Mode = InConfigMode;
+	Mode = ESchemaConfigMode::Edit;
 
-	const FVector2D DefaultWindowSize = FVector2D(1400, 750);
+	ChildSlot [ BuildUI() ];
+}
 
-	FText WindowTitle;
-	if (Mode == ESchemaConfigMode::CreateNew)
+void SCineAssemblySchemaWindow::Construct(const FArguments& InArgs, FGuid InSchemaGuid)
+{
+	Mode = ESchemaConfigMode::Edit;
+
+	// The UI will be temporary because no CineAssemblySchema has been found yet
+	ChildSlot [ BuildUI() ];
+
+	// If the asset registry is still scanning assets, add a callback to find the schema asset matching the input GUID and update this widget once the scan is finished.
+	// Otherwise, we can find the schema asset and update the UI immediately. 
+	const FAssetRegistryModule& AssetRegistryModule = FModuleManager::GetModuleChecked<FAssetRegistryModule>("AssetRegistry");
+	if (AssetRegistryModule.Get().IsLoadingAssets())
 	{
-		WindowTitle = LOCTEXT("WindowTitleCreateNew", "Create Assembly Schema");
+		AssetRegistryModule.Get().OnFilesLoaded().AddSP(this, &SCineAssemblySchemaWindow::FindSchema, InSchemaGuid);
 	}
-	else if (Mode == ESchemaConfigMode::Edit)
+	else
 	{
-		WindowTitle = LOCTEXT("WindowTitleEdit", "Edit Assembly Schema");
+		FindSchema(InSchemaGuid);
+	}
+}
+
+TSharedRef<SWidget> SCineAssemblySchemaWindow::BuildUI()
+{
+	// Build a temporary UI to display while waiting for the schema to be loaded
+	if (!SchemaToConfigure)
+	{
+		return SNew(SBorder)
+			.BorderImage(FCineAssemblyToolsStyle::Get().GetBrush("Borders.PanelNoBorder"))
+			.Padding(8.0f)
+			[
+				SNew(STextBlock).Text(LOCTEXT("LoadingSchemaText", "Loading Cine Assembly Schema..."))
+			];
 	}
 
 	// Initialize the content tree view with the current list of assets and folders saved in the schema
@@ -78,54 +104,100 @@ void SCineAssemblySchemaWindow::Construct(const FArguments& InArgs, UCineAssembl
 	ContentBrowserSettings->SetDisplayEngineFolder(bShowEngineContent);
 	ContentBrowserSettings->SetDisplayPluginFolders(bShowPluginContent);
 
-	GetOnWindowClosedEvent().AddSP(this, &SCineAssemblySchemaWindow::OnSchemaWindowClosed);
-
-	SWindow::Construct(SWindow::FArguments()
-		.Title(WindowTitle)
-		.SizingRule(ESizingRule::UserSized)
-		.ClientSize(DefaultWindowSize)
-		.SupportsMaximize(false)
-		.SupportsMinimize(false)
+	return SNew(SBorder)
+		.BorderImage(FCineAssemblyToolsStyle::Get().GetBrush("ProductionWizard.PanelBackground"))
 		[
-			SNew(SBorder)
-				.BorderImage(FCineAssemblyToolsStyle::Get().GetBrush("ProductionWizard.PanelBackground"))
+			SNew(SVerticalBox)
+
+			+ SVerticalBox::Slot()
 				[
-					SNew(SVerticalBox)
+					SNew(SSplitter)
+						.Orientation(Orient_Horizontal)
+						.PhysicalSplitterHandleSize(2.0f)
 
-					+ SVerticalBox::Slot()
+					+ SSplitter::Slot()
+						.Value(0.2f)
 						[
-							SNew(SSplitter)
-								.Orientation(Orient_Horizontal)
-								.PhysicalSplitterHandleSize(2.0f)
-
-							+ SSplitter::Slot()
-								.Value(0.2f)
-								[
-									MakeMenuPanel()
-								]
-
-							+ SSplitter::Slot()
-								.Value(0.8f)
-								[
-									MakeContentPanel()
-								]
+							MakeMenuPanel()
 						]
 
-					+ SVerticalBox::Slot()
-						.AutoHeight()
+					+ SSplitter::Slot()
+						.Value(0.8f)
 						[
-							SNew(SSeparator)
-								.Orientation(Orient_Horizontal)
-								.Thickness(2.0f)
+							MakeContentPanel()
 						]
+				]
 
-					+ SVerticalBox::Slot()
-						.AutoHeight()
-						[
-							MakeButtonsPanel()
-						]
+			+ SVerticalBox::Slot()
+				.AutoHeight()
+				[
+					SNew(SSeparator)
+						.Orientation(Orient_Horizontal)
+						.Thickness(2.0f)
 				]
-		]);
+
+			+ SVerticalBox::Slot()
+				.AutoHeight()
+				[
+					MakeButtonsPanel()
+				]
+		];
+}
+
+SCineAssemblySchemaWindow::~SCineAssemblySchemaWindow()
+{
+	if (FAssetRegistryModule* AssetRegistryModule = FModuleManager::GetModulePtr<FAssetRegistryModule>("AssetRegistry"))
+	{
+		AssetRegistryModule->Get().OnFilesLoaded().RemoveAll(this);
+	}
+
+	// Save the UI config settings for whether to display engine/plugin content
+	if (UContentBrowserSettings* ContentBrowserSettings = GetMutableDefault<UContentBrowserSettings>())
+	{
+		const bool bShowEngineContent = ContentBrowserSettings->GetDisplayEngineFolder();
+		const bool bShowPluginContent = ContentBrowserSettings->GetDisplayPluginFolders();
+		
+		if (GConfig)
+		{
+			GConfig->SetBool(TEXT("NewCineAssemblySchemaUI"), TEXT("bShowEngineContent"), bShowEngineContent, GEditorPerProjectIni);
+			GConfig->SetBool(TEXT("NewCineAssemblySchemaUI"), TEXT("bShowPluginContent"), bShowPluginContent, GEditorPerProjectIni);
+		}
+
+		ContentBrowserSettings->SetDisplayEngineFolder(bShowEngineContentCached);
+		ContentBrowserSettings->SetDisplayPluginFolders(bShowPluginContentCached);
+	}
+}
+
+void SCineAssemblySchemaWindow::FindSchema(FGuid SchemaID)
+{
+	const FAssetRegistryModule& AssetRegistryModule = FModuleManager::GetModuleChecked<FAssetRegistryModule>("AssetRegistry");
+
+	// The only search criterion for the asset search is for an asset with an SchemaID matching the input GUID
+	const TMultiMap<FName, FString> TagValues = { { UCineAssemblySchema::SchemaGuidPropertyName, SchemaID.ToString() } };
+
+	TArray<FAssetData> SchemaAssets;
+	AssetRegistryModule.Get().GetAssetsByTagValues(TagValues, SchemaAssets);
+
+	// The Schema ID is unique, so at most one asset should ever be found
+	if (SchemaAssets.Num() > 0)
+	{
+		SchemaToConfigure = TStrongObjectPtr<UCineAssemblySchema>(Cast<UCineAssemblySchema>(SchemaAssets[0].GetAsset()));
+
+		// Update the widget's UI
+		ChildSlot.DetachWidget();
+		ChildSlot.AttachWidget(BuildUI());
+	}
+}
+
+FString SCineAssemblySchemaWindow::GetSchemaName()
+{
+	if (SchemaToConfigure)
+	{
+		FString SchemaName;
+		SchemaToConfigure->GetName(SchemaName);
+		return SchemaName;
+	}
+	return TEXT("CineAssemblySchema");
 }
 
 TSharedRef<SWidget> SCineAssemblySchemaWindow::MakeMenuPanel()
@@ -224,8 +296,16 @@ TSharedRef<SWidget> SCineAssemblySchemaWindow::MakeContentPanel()
 									})
 								.OnTextCommitted_Lambda([this](const FText& InText, ETextCommit::Type InCommitType)
 									{
-										SchemaToConfigure->Modify();
-										SchemaToConfigure->SchemaName = InText.ToString();
+										if (SchemaToConfigure->SchemaName != InText.ToString())
+										{
+											SchemaToConfigure->Modify();
+											SchemaToConfigure->SchemaName = InText.ToString();
+
+											if (Mode == ESchemaConfigMode::Edit)
+											{
+												SchemaToConfigure->RenameAsset();
+											}
+										}
 									})
 						]
 				]
@@ -594,6 +674,8 @@ FReply SCineAssemblySchemaWindow::OnAddFolderClicked()
 	// Save a reference to this item so that when the tree is rebuilt, we can immediately start editing its name
 	MostRecentlyAddedItem = NewFolder;
 
+	UpdateFolderList();
+
 	TreeView->SetItemExpansion(SelectedTreeItem, true);
 	TreeView->RequestTreeRefresh();
 
@@ -1037,6 +1119,9 @@ void SCineAssemblySchemaWindow::OnTreeItemTextCommitted(const FText& InText, ETe
 		TreeItem->Parent->ChildAssets.Sort(UE::CineAssemblyTools::Private::SortTreeItems);
 	}
 
+	UpdateFolderList();
+	UpdateAssetList();
+
 	TreeView->RequestTreeRefresh();
 }
 
@@ -1108,6 +1193,7 @@ void SCineAssemblySchemaWindow::DeleteTreeItem(TSharedPtr<FSchemaTreeItem> TreeI
 		if (ContainsRecursive(TreeItem, TopLevelAssemblyItem))
 		{
 			TreeItem->Parent->ChildAssets.Add(TopLevelAssemblyItem);
+			TopLevelAssemblyItem->Parent = TreeItem->Parent;
 		}
 
 		TreeItem->Parent->ChildFolders.Remove(TreeItem);
@@ -1116,6 +1202,10 @@ void SCineAssemblySchemaWindow::DeleteTreeItem(TSharedPtr<FSchemaTreeItem> TreeI
 	{
 		TreeItem->Parent->ChildAssets.Remove(TreeItem);
 	}
+
+	UpdateFolderList();
+	UpdateAssetList();
+
 	TreeView->RequestTreeRefresh();
 }
 
@@ -1143,7 +1233,17 @@ void SCineAssemblySchemaWindow::OnSchemaPropertiesChanged(const FPropertyChanged
 	{
 		if (Mode == ESchemaConfigMode::Edit)
 		{
-			bPendingRename = true;
+			SchemaToConfigure->RenameAsset();
+		}
+	}
+	else if (PropertyName == GET_MEMBER_NAME_CHECKED(FAssemblyMetadataDesc, Key))
+	{
+		UNamingTokensEngineSubsystem* NamingTokensSubsystem = GEngine->GetEngineSubsystem<UNamingTokensEngineSubsystem>();
+		UCineAssemblyNamingTokens* CineAssemblyNamingTokens = Cast<UCineAssemblyNamingTokens>(NamingTokensSubsystem->GetNamingTokens(UCineAssemblyNamingTokens::TokenNamespace));
+
+		for (const FAssemblyMetadataDesc& MetadataDesc : SchemaToConfigure->AssemblyMetadata)
+		{
+			CineAssemblyNamingTokens->AddMetadataToken(MetadataDesc.Key);
 		}
 	}
 }
@@ -1152,35 +1252,28 @@ FReply SCineAssemblySchemaWindow::OnCreateAssetClicked()
 {
 	UCineAssemblySchemaFactory::CreateConfiguredSchema(SchemaToConfigure.Get(), CreateAssetPath);
 
-	RequestDestroyWindow();
-	return FReply::Handled();
-}
+	TSharedPtr<SWindow> ContainingWindow = FSlateApplication::Get().FindWidgetWindow(AsShared());
+	if (ContainingWindow.IsValid())
+	{
+		ContainingWindow->RequestDestroyWindow();
+	}
 
-FReply SCineAssemblySchemaWindow::OnCancelClicked()
-{
-	RequestDestroyWindow();
 	return FReply::Handled();
 }
 
-void SCineAssemblySchemaWindow::OnSchemaWindowClosed(const TSharedRef<SWindow>& InWindow)
+FReply SCineAssemblySchemaWindow::OnCancelClicked()
 {
-	if (bPendingRename)
+	TSharedPtr<SWindow> ContainingWindow = FSlateApplication::Get().FindWidgetWindow(AsShared());
+	if (ContainingWindow.IsValid())
 	{
-		SchemaToConfigure->RenameAsset();
+		ContainingWindow->RequestDestroyWindow();
 	}
 
-	// Add each metadata key to the list of supported naming tokens
-	UNamingTokensEngineSubsystem* NamingTokensSubsystem = GEngine->GetEngineSubsystem<UNamingTokensEngineSubsystem>();
-	UCineAssemblyNamingTokens* CineAssemblyNamingTokens = Cast<UCineAssemblyNamingTokens>(NamingTokensSubsystem->GetNamingTokens(UCineAssemblyNamingTokens::TokenNamespace));
-
-	for (const FAssemblyMetadataDesc& MetadataDesc : SchemaToConfigure->AssemblyMetadata)
-	{
-		CineAssemblyNamingTokens->AddMetadataToken(MetadataDesc.Key);
-	}
-
-	// Set the path of the top-level assembly
-	SchemaToConfigure->DefaultAssemblyPath = TopLevelAssemblyItem->Parent->Path;
+	return FReply::Handled();
+}
 
+void SCineAssemblySchemaWindow::UpdateFolderList()
+{
 	// Update the schema's folder list
 	const TArray<FString> CachedFolderList = SchemaToConfigure->FoldersToCreate;
 
@@ -1192,31 +1285,23 @@ void SCineAssemblySchemaWindow::OnSchemaWindowClosed(const TSharedRef<SWindow>&
 		SchemaToConfigure->Modify();
 		SchemaToConfigure->FoldersToCreate = FolderList;
 	}
+}
 
+void SCineAssemblySchemaWindow::UpdateAssetList()
+{
 	// Update the schema's asset list
 	const TArray<FString> CachedAssetList = SchemaToConfigure->SubsequencesToCreate;
 
 	TArray<FString> AssetList;
 	GetAssetListRecursive(RootItem, AssetList);
 
+	SchemaToConfigure->DefaultAssemblyPath = TopLevelAssemblyItem->Parent->Path;
+
 	if (AssetList != CachedAssetList)
 	{
 		SchemaToConfigure->Modify();
 		SchemaToConfigure->SubsequencesToCreate = AssetList;
 	}
-
-	// Save the UI config settings for whether to display engine/plugin content
-	UContentBrowserSettings* ContentBrowserSettings = GetMutableDefault<UContentBrowserSettings>();
-
-	const bool bShowEngineContent = ContentBrowserSettings->GetDisplayEngineFolder();
-	const bool bShowPluginContent = ContentBrowserSettings->GetDisplayPluginFolders();
-
-	GConfig->SetBool(TEXT("NewCineAssemblySchemaUI"), TEXT("bShowEngineContent"), bShowEngineContent, GEditorPerProjectIni);
-	GConfig->SetBool(TEXT("NewCineAssemblySchemaUI"), TEXT("bShowPluginContent"), bShowPluginContent, GEditorPerProjectIni);
-	GConfig->Flush(false, GEditorPerProjectIni);
-
-	ContentBrowserSettings->SetDisplayEngineFolder(bShowEngineContentCached);
-	ContentBrowserSettings->SetDisplayPluginFolders(bShowPluginContentCached);
 }
 
 FReply SCineAssemblySchemaWindow::OnAssetRowDragDetected(const FGeometry& MyGeometry, const FPointerEvent& MouseEvent)
@@ -1349,6 +1434,9 @@ FReply SCineAssemblySchemaWindow::OnTreeRowAcceptDrop(const FDragDropEvent& InDr
 			InItem->ChildFolders.Sort(UE::CineAssemblyTools::Private::SortTreeItems);
 		}
 
+		UpdateFolderList();
+		UpdateAssetList();
+
 		TreeView->SetItemExpansion(InItem, true);
 		TreeView->RequestTreeRefresh();
 	}
diff --git a/Engine/Plugins/VirtualProduction/CinematicAssemblyTools/Source/CineAssemblyToolsEditor/Private/UI/CineAssembly/SCineAssemblySchemaWindow.h b/Engine/Plugins/VirtualProduction/CinematicAssemblyTools/Source/CineAssemblyToolsEditor/Private/UI/CineAssembly/SCineAssemblySchemaWindow.h
--- a/Engine/Plugins/VirtualProduction/CinematicAssemblyTools/Source/CineAssemblyToolsEditor/Private/UI/CineAssembly/SCineAssemblySchemaWindow.h
+++ b/Engine/Plugins/VirtualProduction/CinematicAssemblyTools/Source/CineAssemblyToolsEditor/Private/UI/CineAssembly/SCineAssemblySchemaWindow.h
@@ -2,7 +2,7 @@
 
 #pragma once
 
-#include "Widgets/SWindow.h"
+#include "Widgets/SCompoundWidget.h"
 
 #include "CineAssemblySchema.h"
 #include "Input/DragAndDrop.h"
@@ -89,18 +89,39 @@ public:
 /**
  * A window for configuring the properties of a UCineAssemblySchema
  */
-class SCineAssemblySchemaWindow : public SWindow
+class SCineAssemblySchemaWindow : public SCompoundWidget
 {
 public:
 	SCineAssemblySchemaWindow() = default;
+	~SCineAssemblySchemaWindow();
 
 	SLATE_BEGIN_ARGS(SCineAssemblySchemaWindow) {}
 	SLATE_END_ARGS()
 
-	void Construct(const FArguments& InArgs, UCineAssemblySchema* InSchema, ESchemaConfigMode InConfigMode);
+	/** Widget construction, initialized with the path where a new schema asset will be created */
 	void Construct(const FArguments& InArgs, const FString& InCreateAssetPath);
 
+	/** Widget construction, initialized with the schema asset being edited */
+	void Construct(const FArguments& InArgs, UCineAssemblySchema* InSchema);
+
+	/**
+	 * Widget construction, initialized with the GUID of the schema to be edited
+	 * The widget will search the asset registry to find the schema asset with the matching GUID,
+	 * and then update the widget contents accordingly.
+	 */
+	void Construct(const FArguments& InArgs, FGuid InSchemaGuid);
+
+	/** Returns the name of the schema asset being edited */
+	FString GetSchemaName();
+
+	/** Searches the asset registry for a Cine Assembly Schema matching the input ID and updates the UI */
+	void FindSchema(FGuid SchemaID);
+
+
 private:
+	/** Constructs the main UI for the widget */
+	TSharedRef<SWidget> BuildUI();
+
 	/** Creates the panel that displays the tab menu */
 	TSharedRef<SWidget> MakeMenuPanel();
 
@@ -191,9 +212,6 @@ private:
 	/** Callback when one of the properties of the schema being configured changes */
 	void OnSchemaPropertiesChanged(const FPropertyChangedEvent& PropertyChangedEvent);
 
-	/** Callback when this window is about to be closed */
-	void OnSchemaWindowClosed(const TSharedRef<SWindow>& InWindow);
-
 	/** Validate the text entered by the user to ensure it will be a valid asset name */
 	bool ValidateAssetName(const FText& InText, FText& OutErrorMessage) const;
 
@@ -221,6 +239,12 @@ private:
 	/** Returns true if the input parent, or any of its child folders, contains in the input item */
 	bool ContainsRecursive(TSharedPtr<FSchemaTreeItem> InParent, TSharedPtr<FSchemaTreeItem> InItem);
 
+	/** Update the list of folders to create from the current set in the tree view */
+	void UpdateFolderList();
+
+	/** Update the list of assets to create from the current set in the tree view */
+	void UpdateAssetList();
+
 	/** Cached content browser settings, used to restore defaults when closing the window */
 	bool bShowEngineContentCached = false;
 	bool bShowPluginContentCached = false;
@@ -241,9 +265,6 @@ private:
 	/** The root path where the configured schema will be created */
 	FString CreateAssetPath;
 
-	/** Whether the schema asset being configured needs to be renamed */
-	bool bPendingRename = false;
-
 	/** List items sources for the asset list view */
 	TArray<TSharedPtr<FString>> AssetListItems;
 
