ShotManagement: Fix Cine Assemblies recorded via Take Recorder to append the Default Assembly Path to the record path location.
#rnx

diff --git a/Engine/Plugins/VirtualProduction/CinematicAssemblyTools/Source/CineAssemblyToolsEditor/Private/TakeRecorder/CineAssemblyTakeRecorderIntegration.cpp b/Engine/Plugins/VirtualProduction/CinematicAssemblyTools/Source/CineAssemblyToolsEditor/Private/TakeRecorder/CineAssemblyTakeRecorderIntegration.cpp
--- a/Engine/Plugins/VirtualProduction/CinematicAssemblyTools/Source/CineAssemblyToolsEditor/Private/TakeRecorder/CineAssemblyTakeRecorderIntegration.cpp
+++ b/Engine/Plugins/VirtualProduction/CinematicAssemblyTools/Source/CineAssemblyToolsEditor/Private/TakeRecorder/CineAssemblyTakeRecorderIntegration.cpp
@@ -3,8 +3,10 @@
 #include "CineAssemblyTakeRecorderIntegration.h"
 
 #include "CineAssembly.h"
+#include "CineAssemblyFactory.h"
 #include "CineAssemblySchema.h"
 #include "CineAssemblyTakeRecorderSettings.h"
+#include "ProductionSettings.h"
 #include "ITakeRecorderModule.h"
 #include "Recorder/TakeRecorder.h"
 #include "Sections/MovieSceneSubSection.h"
@@ -43,6 +45,15 @@ void FCineAssemblyTakeRecorderIntegration::OnRecordingInitialized(UTakeRecorder*
 
 			Assembly->Level = FSoftObjectPath(TakesUtils::DiscoverSourceWorld());
 
+			// Set the production of this recorded subsequence to the current active production
+			const UProductionSettings* ProductionSettings = GetDefault<UProductionSettings>();
+			TOptional<const FCinematicProduction> ActiveProduction = ProductionSettings->GetActiveProduction();
+			if (ActiveProduction.IsSet())
+			{
+				Assembly->Production = ActiveProduction->ProductionID;
+				Assembly->ProductionName = ActiveProduction->ProductionName;
+			}
+
 			// Reset the assembly's playback range, since it was initialized to some default value when it was created, but we do not know what its final range will be until we finish recording.
 			if (UMovieScene* MovieScene = Assembly->GetMovieScene())
 			{
@@ -55,6 +66,23 @@ void FCineAssemblyTakeRecorderIntegration::OnRecordingInitialized(UTakeRecorder*
 			{
 				UCineAssemblySchema* Schema = Settings->AssemblySchema.LoadSynchronous();
 				Assembly->SetSchema(Schema);
+
+				// If the schema defines some additional default assembly path, move the assembly asset accordingly
+				if (!Schema->DefaultAssemblyPath.IsEmpty())
+				{
+					FString CurrentAssemblyPath;
+					Assembly->GetPathName(nullptr, CurrentAssemblyPath);
+					CurrentAssemblyPath = FPaths::GetPath(CurrentAssemblyPath);
+
+					FString UniquePackageName;
+					FString UniqueAssetName;
+					UCineAssemblyFactory::MakeUniqueNameAndPath(Assembly, CurrentAssemblyPath, UniquePackageName, UniqueAssetName);
+
+					UPackage* Package = CreatePackage(*UniquePackageName);
+					Assembly->Rename(*UniqueAssetName, Package);
+				}
+
+				Assembly->CreateSubAssemblies();
 			}
 		}
 	}
@@ -66,8 +94,6 @@ void FCineAssemblyTakeRecorderIntegration::OnRecordingStarted(UTakeRecorder* Tak
 	{
 		if (UCineAssembly* Assembly = Cast<UCineAssembly>(TakeRecorder->GetSequence()))
 		{
-			Assembly->CreateSubAssemblies();
-
 			// Reset the playback range of each subsection and its underlying sequence's movie scene.
 			// These were initialized to default ranges when the assets were created, but when recording, we do not want to use the default range.
 			for (UMovieSceneSubSection* SubSection : Assembly->SubAssemblies)
