ShotManagement: Split the assembly metadata type "number" into two distinct types "integer" and "float" to better represent different types of data.
#rnx

#virtualized

diff --git a/Engine/Plugins/VirtualProduction/CinematicAssemblyTools/Source/CineAssemblyTools/Private/CineAssembly.cpp b/Engine/Plugins/VirtualProduction/CinematicAssemblyTools/Source/CineAssemblyTools/Private/CineAssembly.cpp
--- a/Engine/Plugins/VirtualProduction/CinematicAssemblyTools/Source/CineAssemblyTools/Private/CineAssembly.cpp
+++ b/Engine/Plugins/VirtualProduction/CinematicAssemblyTools/Source/CineAssemblyTools/Private/CineAssembly.cpp
@@ -120,9 +120,14 @@ void UCineAssembly::ChangeSchema(UCineAssemblySchema* InSchema)
 				const bool& DefaultValue = MetadataDesc.DefaultValue.Get<bool>();
 				MetadataJsonObject->SetBoolField(MetadataDesc.Key, DefaultValue);
 			}
-			else if (MetadataDesc.DefaultValue.IsType<double>())
+			else if (MetadataDesc.DefaultValue.IsType<int32>())
 			{
-				const double& DefaultValue = MetadataDesc.DefaultValue.Get<double>();
+				const int32& DefaultValue = MetadataDesc.DefaultValue.Get<int32>();
+				MetadataJsonObject->SetNumberField(MetadataDesc.Key, DefaultValue);
+			}
+			else if (MetadataDesc.DefaultValue.IsType<float>())
+			{
+				const float& DefaultValue = MetadataDesc.DefaultValue.Get<float>();
 				MetadataJsonObject->SetNumberField(MetadataDesc.Key, DefaultValue);
 			}
 		}
@@ -368,7 +373,14 @@ void UCineAssembly::SetMetadataAsBool(FString InKey, bool InValue)
 	AddMetadataNamingToken(InKey);
 }
 
-void UCineAssembly::SetMetadataAsNumber(FString InKey, double InValue)
+void UCineAssembly::SetMetadataAsInteger(FString InKey, int32 InValue)
+{
+	Modify();
+	MetadataJsonObject->SetNumberField(InKey, InValue);
+	AddMetadataNamingToken(InKey);
+}
+
+void UCineAssembly::SetMetadataAsFloat(FString InKey, float InValue)
 {
 	Modify();
 	MetadataJsonObject->SetNumberField(InKey, InValue);
@@ -395,7 +407,17 @@ bool UCineAssembly::GetMetadataAsBool(FString InKey, bool& OutValue) const
 	return true;
 }
 
-bool UCineAssembly::GetMetadataAsNumber(FString InKey, double& OutValue) const
+bool UCineAssembly::GetMetadataAsInteger(FString InKey, int32& OutValue) const
+{
+	if (!MetadataJsonObject->TryGetNumberField(InKey, OutValue))
+	{
+		OutValue = 0;
+		return false;
+	}
+	return true;
+}
+
+bool UCineAssembly::GetMetadataAsFloat(FString InKey, float& OutValue) const
 {
 	if (!MetadataJsonObject->TryGetNumberField(InKey, OutValue))
 	{
diff --git a/Engine/Plugins/VirtualProduction/CinematicAssemblyTools/Source/CineAssemblyTools/Public/CineAssembly.h b/Engine/Plugins/VirtualProduction/CinematicAssemblyTools/Source/CineAssemblyTools/Public/CineAssembly.h
--- a/Engine/Plugins/VirtualProduction/CinematicAssemblyTools/Source/CineAssemblyTools/Public/CineAssembly.h
+++ b/Engine/Plugins/VirtualProduction/CinematicAssemblyTools/Source/CineAssemblyTools/Public/CineAssembly.h
@@ -89,9 +89,13 @@ public:
 	UFUNCTION(BlueprintCallable, Category = "Cine Assembly Tools")
 	CINEASSEMBLYTOOLS_API void SetMetadataAsBool(FString InKey, bool InValue);
 
-	/** Add a double-precision floating point number as metadata to this assembly */
+	/** Add an integer as metadata to this assembly */
 	UFUNCTION(BlueprintCallable, Category = "Cine Assembly Tools")
-	CINEASSEMBLYTOOLS_API void SetMetadataAsNumber(FString InKey, double InValue);
+	CINEASSEMBLYTOOLS_API void SetMetadataAsInteger(FString InKey, int32 InValue);
+
+	/** Add a floating point number as metadata to this assembly */
+	UFUNCTION(BlueprintCallable, Category = "Cine Assembly Tools")
+	CINEASSEMBLYTOOLS_API void SetMetadataAsFloat(FString InKey, float InValue);
 
 	/** Get the metadata value for the input key as a string (if it exists) */
 	UFUNCTION(BlueprintPure, Category = "Cine Assembly Tools", meta = (ReturnDisplayName = "FoundKey"))
@@ -101,9 +105,13 @@ public:
 	UFUNCTION(BlueprintPure, Category = "Cine Assembly Tools", meta = (ReturnDisplayName = "FoundKey"))
 	CINEASSEMBLYTOOLS_API bool GetMetadataAsBool(FString InKey, bool& OutValue) const;
 
-	/** Get the metadata value for the input key as a double-precision floating-point number (if it exists) */
+	/** Get the metadata value for the input key as an integer (if it exists) */
+	UFUNCTION(BlueprintPure, Category = "Cine Assembly Tools", meta = (ReturnDisplayName = "FoundKey"))
+	CINEASSEMBLYTOOLS_API bool GetMetadataAsInteger(FString InKey, int32& OutValue) const;
+
+	/** Get the metadata value for the input key as a floating-point number (if it exists) */
 	UFUNCTION(BlueprintPure, Category = "Cine Assembly Tools", meta = (ReturnDisplayName = "FoundKey"))
-	CINEASSEMBLYTOOLS_API bool GetMetadataAsNumber(FString InKey, double& OutValue) const;
+	CINEASSEMBLYTOOLS_API bool GetMetadataAsFloat(FString InKey, float& OutValue) const;
 
 	/** Adds a new metadata key to the list of supported naming tokens for assemblies */
 	CINEASSEMBLYTOOLS_API void AddMetadataNamingToken(const FString& InKey);
diff --git a/Engine/Plugins/VirtualProduction/CinematicAssemblyTools/Source/CineAssemblyTools/Public/CineAssemblySchema.h b/Engine/Plugins/VirtualProduction/CinematicAssemblyTools/Source/CineAssemblyTools/Public/CineAssemblySchema.h
--- a/Engine/Plugins/VirtualProduction/CinematicAssemblyTools/Source/CineAssemblyTools/Public/CineAssemblySchema.h
+++ b/Engine/Plugins/VirtualProduction/CinematicAssemblyTools/Source/CineAssemblyTools/Public/CineAssemblySchema.h
@@ -13,7 +13,8 @@ enum class ECineAssemblyMetadataType : uint8
 {
 	String = 0,
 	Bool,
-	Number,
+	Integer,
+	Float,
 	AssetPath,
 	CineAssembly
 };
@@ -41,7 +42,7 @@ struct FAssemblyMetadataDesc
 	FSoftObjectPath SchemaType;
 
 	/** The default value for this metadata field */
-	TVariant<FString, bool, double> DefaultValue;
+	TVariant<FString, bool, int32, float> DefaultValue;
 };
 
 /**
diff --git a/Engine/Plugins/VirtualProduction/CinematicAssemblyTools/Source/CineAssemblyToolsEditor/Private/CineAssemblyCustomization.cpp b/Engine/Plugins/VirtualProduction/CinematicAssemblyTools/Source/CineAssemblyToolsEditor/Private/CineAssemblyCustomization.cpp
--- a/Engine/Plugins/VirtualProduction/CinematicAssemblyTools/Source/CineAssemblyToolsEditor/Private/CineAssemblyCustomization.cpp
+++ b/Engine/Plugins/VirtualProduction/CinematicAssemblyTools/Source/CineAssemblyToolsEditor/Private/CineAssemblyCustomization.cpp
@@ -173,22 +173,40 @@ void FCineAssemblyCustomization::CustomizeMetadataCategory(IDetailLayoutBuilder&
 						CustomizedCineAssembly->SetMetadataAsBool(MetadataDesc.Key, Value);
 					});
 		}
-		else if (MetadataDesc.Type == ECineAssemblyMetadataType::Number)
+		else if (MetadataDesc.Type == ECineAssemblyMetadataType::Integer)
 		{
-			ValueWidget = SNew(SNumericEntryBox<double>)
+			ValueWidget = SNew(SNumericEntryBox<int32>)
 				.Font(FAppStyle::GetFontStyle(TEXT("PropertyWindow.NormalFont")))
 				.Value_Lambda([this, MetadataDesc]()
 					{
-						double Value;
-						if (!CustomizedCineAssembly->GetMetadataAsNumber(MetadataDesc.Key, Value))
+						int32 Value;
+						if (!CustomizedCineAssembly->GetMetadataAsInteger(MetadataDesc.Key, Value))
 						{
-							Value = MetadataDesc.DefaultValue.Get<double>();
+							Value = MetadataDesc.DefaultValue.Get<int32>();
 						}
 						return Value;
 					})
-				.OnValueChanged_Lambda([this, MetadataDesc](double InValue)
+				.OnValueChanged_Lambda([this, MetadataDesc](int32 InValue)
 					{
-						CustomizedCineAssembly->SetMetadataAsNumber(MetadataDesc.Key, InValue);
+						CustomizedCineAssembly->SetMetadataAsInteger(MetadataDesc.Key, InValue);
+					});
+		}
+		else if (MetadataDesc.Type == ECineAssemblyMetadataType::Float)
+		{
+			ValueWidget = SNew(SNumericEntryBox<float>)
+				.Font(FAppStyle::GetFontStyle(TEXT("PropertyWindow.NormalFont")))
+				.Value_Lambda([this, MetadataDesc]()
+					{
+						float Value;
+						if (!CustomizedCineAssembly->GetMetadataAsFloat(MetadataDesc.Key, Value))
+						{
+							Value = MetadataDesc.DefaultValue.Get<float>();
+						}
+						return Value;
+					})
+				.OnValueChanged_Lambda([this, MetadataDesc](float InValue)
+					{
+						CustomizedCineAssembly->SetMetadataAsFloat(MetadataDesc.Key, InValue);
 					});
 		}
 		else if (MetadataDesc.Type == ECineAssemblyMetadataType::AssetPath)
diff --git a/Engine/Plugins/VirtualProduction/CinematicAssemblyTools/Source/CineAssemblyToolsEditor/Private/CineAssemblyMetadataCustomization.cpp b/Engine/Plugins/VirtualProduction/CinematicAssemblyTools/Source/CineAssemblyToolsEditor/Private/CineAssemblyMetadataCustomization.cpp
--- a/Engine/Plugins/VirtualProduction/CinematicAssemblyTools/Source/CineAssemblyToolsEditor/Private/CineAssemblyMetadataCustomization.cpp
+++ b/Engine/Plugins/VirtualProduction/CinematicAssemblyTools/Source/CineAssemblyToolsEditor/Private/CineAssemblyMetadataCustomization.cpp
@@ -91,9 +91,13 @@ void FCineAssemblyMetadataCustomization::CustomizeChildren(TSharedRef<IPropertyH
 				{
 					MetadataDesc.DefaultValue.Set<bool>(false);
 				}
-				else if (MetadataDesc.Type == ECineAssemblyMetadataType::Number)
+				else if (MetadataDesc.Type == ECineAssemblyMetadataType::Integer)
 				{
-					MetadataDesc.DefaultValue.Set<double>(0);
+					MetadataDesc.DefaultValue.Set<int32>(0);
+				}
+				else if (MetadataDesc.Type == ECineAssemblyMetadataType::Float)
+				{
+					MetadataDesc.DefaultValue.Set<float>(0);
 				}
 				else
 				{
@@ -151,21 +155,39 @@ void FCineAssemblyMetadataCustomization::CustomizeChildren(TSharedRef<IPropertyH
 						})
 			];
 
-		DefaultValueWidget->AddSlot((uint8)ECineAssemblyMetadataType::Number)
+		DefaultValueWidget->AddSlot((uint8)ECineAssemblyMetadataType::Integer)
+			[
+				SNew(SNumericEntryBox<int32>)
+					.Font(FAppStyle::GetFontStyle(TEXT("PropertyWindow.NormalFont")))
+					.Value_Lambda([&MetadataDesc]() -> int32
+						{
+							if (const int32* Value = MetadataDesc.DefaultValue.TryGet<int32>())
+							{
+								return *Value;
+							}
+							return 0;
+						})
+					.OnValueChanged_Lambda([&MetadataDesc](int32 InValue)
+						{
+							MetadataDesc.DefaultValue.Set<int32>(InValue);
+						})
+			];
+
+		DefaultValueWidget->AddSlot((uint8)ECineAssemblyMetadataType::Float)
 			[
-				SNew(SNumericEntryBox<double>)
+				SNew(SNumericEntryBox<float>)
 					.Font(FAppStyle::GetFontStyle(TEXT("PropertyWindow.NormalFont")))
-					.Value_Lambda([&MetadataDesc]() -> double
+					.Value_Lambda([&MetadataDesc]() -> float
 						{
-							if (const double* Value = MetadataDesc.DefaultValue.TryGet<double>())
+							if (const float* Value = MetadataDesc.DefaultValue.TryGet<float>())
 							{
 								return *Value;
 							}
 							return 0;
 						})
-					.OnValueChanged_Lambda([&MetadataDesc](double InValue)
+					.OnValueChanged_Lambda([&MetadataDesc](float InValue)
 						{
-							MetadataDesc.DefaultValue.Set<double>(InValue);
+							MetadataDesc.DefaultValue.Set<float>(InValue);
 						})
 			];
 
diff --git a/Engine/Plugins/VirtualProduction/CinematicAssemblyTools/Source/CineAssemblyToolsEditor/Private/ProductionFunctionLibrary.cpp b/Engine/Plugins/VirtualProduction/CinematicAssemblyTools/Source/CineAssemblyToolsEditor/Private/ProductionFunctionLibrary.cpp
--- a/Engine/Plugins/VirtualProduction/CinematicAssemblyTools/Source/CineAssemblyToolsEditor/Private/ProductionFunctionLibrary.cpp
+++ b/Engine/Plugins/VirtualProduction/CinematicAssemblyTools/Source/CineAssemblyToolsEditor/Private/ProductionFunctionLibrary.cpp
@@ -121,9 +121,13 @@ void UProductionFunctionLibrary::AddMetadataToAssembly(UCineAssembly* Assembly,
 			{
 				Assembly->SetMetadataAsBool(Pair.Key, Pair.Value.ToBool());
 			}
-			else if (MetadataDesc->Type == ECineAssemblyMetadataType::Number)
+			else if (MetadataDesc->Type == ECineAssemblyMetadataType::Integer)
 			{
-				Assembly->SetMetadataAsNumber(Pair.Key, FCString::Atof(*Pair.Value));
+				Assembly->SetMetadataAsInteger(Pair.Key, FCString::Atoi(*Pair.Value));
+			}
+			else if (MetadataDesc->Type == ECineAssemblyMetadataType::Float)
+			{
+				Assembly->SetMetadataAsFloat(Pair.Key, FCString::Atof(*Pair.Value));
 			}
 			else if (MetadataDesc->Type == ECineAssemblyMetadataType::AssetPath)
 			{
