ShotManagement: Change Cine Assermbly asset color. Update plugin description.
#rnx

diff --git a/Engine/Plugins/VirtualProduction/CinematicAssemblyTools/CinematicAssemblyTools.uplugin b/Engine/Plugins/VirtualProduction/CinematicAssemblyTools/CinematicAssemblyTools.uplugin
--- a/Engine/Plugins/VirtualProduction/CinematicAssemblyTools/CinematicAssemblyTools.uplugin
+++ b/Engine/Plugins/VirtualProduction/CinematicAssemblyTools/CinematicAssemblyTools.uplugin
@@ -2,8 +2,8 @@
 	"FileVersion": 3,
 	"Version": 1,
 	"VersionName": "1.0",
-	"FriendlyName": "Cinematic Assembly Tools",
-	"Description": "Cinematic pipeline tools for shot management and linear content creation",
+	"FriendlyName": "Cinematic Assembly Tools (CAT)",
+	"Description": "CAT is a suite of cinematic pipeline tools for shot management and linear content creation",
 	"Category": "Virtual Production",
 	"CreatedBy": "Epic Games, Inc.",
 	"CreatedByURL": "https://epicgames.com",
diff --git a/Engine/Plugins/VirtualProduction/CinematicAssemblyTools/Source/CineAssemblyToolsEditor/Private/AssetDefinition_CineAssembly.cpp b/Engine/Plugins/VirtualProduction/CinematicAssemblyTools/Source/CineAssemblyToolsEditor/Private/AssetDefinition_CineAssembly.cpp
--- a/Engine/Plugins/VirtualProduction/CinematicAssemblyTools/Source/CineAssemblyToolsEditor/Private/AssetDefinition_CineAssembly.cpp
+++ b/Engine/Plugins/VirtualProduction/CinematicAssemblyTools/Source/CineAssemblyToolsEditor/Private/AssetDefinition_CineAssembly.cpp
@@ -47,7 +47,7 @@ TSoftClassPtr<> UAssetDefinition_CineAssembly::GetAssetClass() const
 
 FLinearColor UAssetDefinition_CineAssembly::GetAssetColor() const
 {
-	return FColor(172, 35, 85);
+	return FColor(229, 45, 113);
 }
 
 TConstArrayView<FAssetCategoryPath> UAssetDefinition_CineAssembly::GetAssetCategories() const
diff --git a/Engine/Plugins/VirtualProduction/CinematicAssemblyTools/Source/CineAssemblyToolsEditor/Private/UI/SAssetNamingPanel.cpp b/Engine/Plugins/VirtualProduction/CinematicAssemblyTools/Source/CineAssemblyToolsEditor/Private/UI/SAssetNamingPanel.cpp
--- a/Engine/Plugins/VirtualProduction/CinematicAssemblyTools/Source/CineAssemblyToolsEditor/Private/UI/SAssetNamingPanel.cpp
+++ b/Engine/Plugins/VirtualProduction/CinematicAssemblyTools/Source/CineAssemblyToolsEditor/Private/UI/SAssetNamingPanel.cpp
@@ -143,7 +143,7 @@ void SAssetNamingPanel::Construct(const FArguments& InArgs)
 								.Padding(0.0f, 0.0f, 0.0f, 16.0f)
 								[
 									SNew(STextBlock)
-										.Text(LOCTEXT("AssetNamingInfoText", "For each asset class, you can set a prefix for the editor to assign automatically to the names of newly created assets. This helps keep asset names consistent across your project."))
+										.Text(LOCTEXT("AssetNamingInfoText", "For each asset class, you can set a name to be used automatically for newly created assets. This helps keep asset names consistent across your project."))
 										.AutoWrapText(true)
 								]
 
