ShotManagement: Add BP API to get/set the assembly production and parent assembly
#rnx

diff --git a/Engine/Plugins/VirtualProduction/CinematicAssemblyTools/Source/CineAssemblyTools/Private/CineAssembly.cpp b/Engine/Plugins/VirtualProduction/CinematicAssemblyTools/Source/CineAssemblyTools/Private/CineAssembly.cpp
--- a/Engine/Plugins/VirtualProduction/CinematicAssemblyTools/Source/CineAssemblyTools/Private/CineAssembly.cpp
+++ b/Engine/Plugins/VirtualProduction/CinematicAssemblyTools/Source/CineAssemblyTools/Private/CineAssembly.cpp
@@ -349,6 +349,30 @@ void UCineAssembly::AppendToNoteText(FString InNote)
 	AssemblyNote.Append(InNote);
 }
 
+FGuid UCineAssembly::GetProductionID()
+{
+	return Production;
+}
+
+FString UCineAssembly::GetProductionName()
+{
+	return ProductionName;
+}
+
+TSoftObjectPtr<UCineAssembly> UCineAssembly::GetParentAssembly()
+{
+	if (ParentAssembly.IsValid())
+	{
+		return TSoftObjectPtr<UCineAssembly>(ParentAssembly);
+	}
+	return nullptr;
+}
+
+void UCineAssembly::SetParentAssembly(TSoftObjectPtr<UCineAssembly> InParent)
+{
+	ParentAssembly = InParent.ToSoftObjectPath();
+}
+
 FString UCineAssembly::GetFullMetadataString() const
 {
 	FString JsonString;
diff --git a/Engine/Plugins/VirtualProduction/CinematicAssemblyTools/Source/CineAssemblyTools/Public/CineAssembly.h b/Engine/Plugins/VirtualProduction/CinematicAssemblyTools/Source/CineAssemblyTools/Public/CineAssembly.h
--- a/Engine/Plugins/VirtualProduction/CinematicAssemblyTools/Source/CineAssemblyTools/Public/CineAssembly.h
+++ b/Engine/Plugins/VirtualProduction/CinematicAssemblyTools/Source/CineAssemblyTools/Public/CineAssembly.h
@@ -77,6 +77,22 @@ public:
 	UFUNCTION(BlueprintCallable, Category = "Cine Assembly Tools")
 	CINEASSEMBLYTOOLS_API void AppendToNoteText(FString InNote);
 
+	/** Get the production ID associated with this assembly */
+	UFUNCTION(BlueprintPure, Category = "Cine Assembly Tools")
+	CINEASSEMBLYTOOLS_API FGuid GetProductionID();
+
+	/** Get the production name associated with this assembly */
+	UFUNCTION(BlueprintPure, Category = "Cine Assembly Tools")
+	CINEASSEMBLYTOOLS_API FString GetProductionName();
+
+	/** Get the parent assembly of this assembly */
+	UFUNCTION(BlueprintPure, Category = "Cine Assembly Tools")
+	CINEASSEMBLYTOOLS_API TSoftObjectPtr<UCineAssembly> GetParentAssembly();
+
+	/** Set the parent assembly of this assembly */
+	UFUNCTION(BlueprintCallable, Category = "Cine Assembly Tools")
+	CINEASSEMBLYTOOLS_API void SetParentAssembly(TSoftObjectPtr<UCineAssembly> InParent);
+
 	/** Get all of the metadata for this assembly as a formatted JSON string */
 	UFUNCTION(BlueprintPure, Category = "Cine Assembly Tools")
 	CINEASSEMBLYTOOLS_API FString GetFullMetadataString() const;
diff --git a/Engine/Plugins/VirtualProduction/CinematicAssemblyTools/Source/CineAssemblyToolsEditor/Private/ProductionFunctionLibrary.cpp b/Engine/Plugins/VirtualProduction/CinematicAssemblyTools/Source/CineAssemblyToolsEditor/Private/ProductionFunctionLibrary.cpp
--- a/Engine/Plugins/VirtualProduction/CinematicAssemblyTools/Source/CineAssemblyToolsEditor/Private/ProductionFunctionLibrary.cpp
+++ b/Engine/Plugins/VirtualProduction/CinematicAssemblyTools/Source/CineAssemblyToolsEditor/Private/ProductionFunctionLibrary.cpp
@@ -84,12 +84,22 @@ void UProductionFunctionLibrary::RenameProduction(FGuid ProductionID, FString Ne
 	ProductionSettings->RenameProduction(ProductionID, NewName);
 }
 
-UCineAssembly* UProductionFunctionLibrary::CreateAssembly(UCineAssemblySchema* Schema, TSoftObjectPtr<UWorld> Level, TMap<FString, FString> Metadata, const FString& Path, const FString& Name, bool bUseDefaultNameFromSchema)
+UCineAssembly* UProductionFunctionLibrary::CreateAssembly(UCineAssemblySchema* Schema, TSoftObjectPtr<UWorld> Level, TSoftObjectPtr<UCineAssembly> ParentAssembly, TMap<FString, FString> Metadata, const FString& Path, const FString& Name, bool bUseDefaultNameFromSchema)
 {
 	UCineAssembly* NewAssembly = NewObject<UCineAssembly>(GetTransientPackage(), NAME_None, RF_Transient);
 	NewAssembly->SetSchema(Schema);
 	NewAssembly->SetLevel(Level);
-	
+	NewAssembly->SetParentAssembly(ParentAssembly);
+
+	// Associate the current active production with this assembly
+	const UProductionSettings* const ProductionSettings = GetDefault<UProductionSettings>();
+	TOptional<const FCinematicProduction> ActiveProduction = ProductionSettings->GetActiveProduction();
+	if (ActiveProduction.IsSet())
+	{
+		NewAssembly->Production = ActiveProduction->ProductionID;
+		NewAssembly->ProductionName = ActiveProduction->ProductionName;
+	}
+
 	AddMetadataToAssembly(NewAssembly, Metadata);
 
 	if (!Schema || !bUseDefaultNameFromSchema)
diff --git a/Engine/Plugins/VirtualProduction/CinematicAssemblyTools/Source/CineAssemblyToolsEditor/Private/ProductionFunctionLibrary.h b/Engine/Plugins/VirtualProduction/CinematicAssemblyTools/Source/CineAssemblyToolsEditor/Private/ProductionFunctionLibrary.h
--- a/Engine/Plugins/VirtualProduction/CinematicAssemblyTools/Source/CineAssemblyToolsEditor/Private/ProductionFunctionLibrary.h
+++ b/Engine/Plugins/VirtualProduction/CinematicAssemblyTools/Source/CineAssemblyToolsEditor/Private/ProductionFunctionLibrary.h
@@ -69,7 +69,7 @@ public:
 	 * It is important that any metadata required for resolving asset naming tokens is provided to this function so that the Assembly and SubAssemblies are all named correctly. 
 	 */
 	UFUNCTION(BlueprintCallable, Category = "Cine Assembly Tools")
-	static UCineAssembly* CreateAssembly(UCineAssemblySchema* Schema, TSoftObjectPtr<UWorld> Level, TMap<FString, FString> Metadata, const FString& Path, const FString& Name, bool bUseDefaultNameFromSchema = true);
+	static UCineAssembly* CreateAssembly(UCineAssemblySchema* Schema, TSoftObjectPtr<UWorld> Level, TSoftObjectPtr<UCineAssembly> ParentAssembly, TMap<FString, FString> Metadata, const FString& Path, const FString& Name, bool bUseDefaultNameFromSchema = true);
 
 private:
 	/** Adds the input map of metadata to the input assembly */
