ShotManagement: Mark subassemblies with the IsSubsequence metadata tag.
#rnx

diff --git a/Engine/Plugins/VirtualProduction/CinematicAssemblyTools/Source/CineAssemblyTools/Private/CineAssembly.cpp b/Engine/Plugins/VirtualProduction/CinematicAssemblyTools/Source/CineAssemblyTools/Private/CineAssembly.cpp
--- a/Engine/Plugins/VirtualProduction/CinematicAssemblyTools/Source/CineAssemblyTools/Private/CineAssembly.cpp
+++ b/Engine/Plugins/VirtualProduction/CinematicAssemblyTools/Source/CineAssemblyTools/Private/CineAssembly.cpp
@@ -5,6 +5,7 @@
 #include "AssetRegistry/AssetRegistryModule.h"
 #include "CineAssemblyNamingTokens.h"
 #include "HAL/FileManager.h"
+#include "LevelSequenceShotMetaDataLibrary.h"
 #include "MovieScene.h"
 #include "NamingTokensEngineSubsystem.h"
 #include "Sections/MovieSceneSubSection.h"
@@ -237,6 +238,9 @@ void UCineAssembly::CreateSubAssemblies()
 			SubAssembly->Level = Level;
 			SubAssembly->ParentAssembly = this;
 			SubAssembly->Production = Production;
+			SubAssembly->ProductionName = ProductionName;
+
+			ULevelSequenceShotMetaDataLibrary::SetIsSubSequence(SubAssembly, true);
 
 			const FFrameNumber StartFrame = PlaybackRange.GetLowerBoundValue();
 			const int32 Duration = PlaybackRange.Size<FFrameNumber>().Value;
