ShotManagement: Fix cine assembly to not allow registering empty asset registry keys
#rnx

diff --git a/Engine/Plugins/VirtualProduction/CinematicAssemblyTools/Source/CineAssemblyTools/Private/CineAssembly.cpp b/Engine/Plugins/VirtualProduction/CinematicAssemblyTools/Source/CineAssemblyTools/Private/CineAssembly.cpp
--- a/Engine/Plugins/VirtualProduction/CinematicAssemblyTools/Source/CineAssemblyTools/Private/CineAssembly.cpp
+++ b/Engine/Plugins/VirtualProduction/CinematicAssemblyTools/Source/CineAssemblyTools/Private/CineAssembly.cpp
@@ -262,10 +262,13 @@ void UCineAssembly::GetAssetRegistryTags(FAssetRegistryTagsContext Context) cons
 	// Add tags associated with the assembly metadata
 	for (const TPair<FString, TSharedPtr<FJsonValue>>& Pair : MetadataJsonObject->Values)
 	{
-		FString ValueString;
-		if (MetadataJsonObject->TryGetStringField(Pair.Key, ValueString))
+		if (!Pair.Key.IsEmpty())
 		{
-			Context.AddTag(FAssetRegistryTag(*Pair.Key, ValueString, FAssetRegistryTag::ETagType::TT_Alphabetical, FAssetRegistryTag::TD_None));
+			FString ValueString;
+			if (MetadataJsonObject->TryGetStringField(Pair.Key, ValueString))
+			{
+				Context.AddTag(FAssetRegistryTag(*Pair.Key, ValueString, FAssetRegistryTag::ETagType::TT_Alphabetical, FAssetRegistryTag::TD_None));
+			}
 		}
 	}
 }
