ShotManagement: Fix assembly and schemas to allow opening more than one at a time from content browser
#rnx

diff --git a/Engine/Plugins/VirtualProduction/CinematicAssemblyTools/Source/CineAssemblyToolsEditor/Private/AssetDefinition_CineAssemblySchema.cpp b/Engine/Plugins/VirtualProduction/CinematicAssemblyTools/Source/CineAssemblyToolsEditor/Private/AssetDefinition_CineAssemblySchema.cpp
--- a/Engine/Plugins/VirtualProduction/CinematicAssemblyTools/Source/CineAssemblyToolsEditor/Private/AssetDefinition_CineAssemblySchema.cpp
+++ b/Engine/Plugins/VirtualProduction/CinematicAssemblyTools/Source/CineAssemblyToolsEditor/Private/AssetDefinition_CineAssemblySchema.cpp
@@ -54,10 +54,8 @@ FAssetSupportResponse UAssetDefinition_CineAssemblySchema::CanRename(const FAsse
 
 EAssetCommandResult UAssetDefinition_CineAssemblySchema::OpenAssets(const FAssetOpenArgs& OpenArgs) const
 {
-	// We only support opening one asset at a time
-	if (OpenArgs.Assets.Num() == 1)
+	for (const FAssetData& SchemaAssetData : OpenArgs.Assets)
 	{
-		const FAssetData& SchemaAssetData = OpenArgs.Assets[0];
 		if (UCineAssemblySchema* CineAssemblySchema = Cast<UCineAssemblySchema>(SchemaAssetData.GetAsset()))
 		{
 			FCineAssemblyToolsEditorModule& CineAssemblyToolsEditorModule = FModuleManager::GetModuleChecked<FCineAssemblyToolsEditorModule>("CineAssemblyToolsEditor");
diff --git a/Engine/Plugins/VirtualProduction/CinematicAssemblyTools/Source/CineAssemblyToolsEditor/Private/CineAssemblyToolsEditorModule.cpp b/Engine/Plugins/VirtualProduction/CinematicAssemblyTools/Source/CineAssemblyToolsEditor/Private/CineAssemblyToolsEditorModule.cpp
--- a/Engine/Plugins/VirtualProduction/CinematicAssemblyTools/Source/CineAssemblyToolsEditor/Private/CineAssemblyToolsEditorModule.cpp
+++ b/Engine/Plugins/VirtualProduction/CinematicAssemblyTools/Source/CineAssemblyToolsEditor/Private/CineAssemblyToolsEditorModule.cpp
@@ -2,6 +2,7 @@
 
 #include "CineAssemblyToolsEditorModule.h"
 
+#include "Algo/AllOf.h"
 #include "AssetRegistry/AssetRegistryModule.h"
 #include "AssetToolsModule.h"
 #include "CineAssembly.h"
@@ -262,31 +263,43 @@ void FCineAssemblyToolsEditorModule::RegisterMenus()
 		{
 			if (UContentBrowserAssetContextMenuContext* Context = InSection.FindContext<UContentBrowserAssetContextMenuContext>())
 			{
-				if ((Context->SelectedAssets.Num() == 1) && Context->SelectedAssets[0].IsInstanceOf(UCineAssembly::StaticClass()))
+				// Ensure all selected assets are Cine Assemblies
+				const TArray<FAssetData>& SelectedAssets = Context->SelectedAssets;
+				if (Algo::AllOf(SelectedAssets, [](const FAssetData& AssetData) { return AssetData.IsInstanceOf(UCineAssembly::StaticClass()); }))
 				{
-					const FAssetData& CineAssemblyData = Context->SelectedAssets[0];
-					UCineAssembly* CineAssembly = Cast<UCineAssembly>(CineAssemblyData.GetAsset());
-
-					InSection.AddMenuEntry(
-						"OpenInCurrentMapEntry",
-						LOCTEXT("OpenInCurrentMap", "Open in Current Map"),
-						LOCTEXT("OpenInCurrentMapTooltip", "Opens the level sequence in Sequencer but does not automatically load the associated map"),
-						FSlateIcon(FAppStyle::GetAppStyleSetName(), "LevelEditor.OpenCinematic"),
-						FUIAction(FExecuteAction::CreateLambda([CineAssembly]()
-						{
-							// Note: Normally, when a cine assembly is opened in the editor, its associated map is loaded during the PrepareToActiveAssets step.
-							// By directly opening the asset editor (Sequencer) for this asset, we can skip that map loading step. 
-							GEditor->GetEditorSubsystem<UAssetEditorSubsystem>()->OpenEditorForAssets({ CineAssembly });
-						})),
-						EUserInterfaceActionType::None
-					);
+					if (SelectedAssets.Num() == 1)
+					{
+						const FAssetData& CineAssemblyData = SelectedAssets[0];
+						UCineAssembly* CineAssembly = Cast<UCineAssembly>(CineAssemblyData.GetAsset());
+
+						InSection.AddMenuEntry(
+							"OpenInCurrentMapEntry",
+							LOCTEXT("OpenInCurrentMap", "Open in Current Map"),
+							LOCTEXT("OpenInCurrentMapTooltip", "Opens the level sequence in Sequencer but does not automatically load the associated map"),
+							FSlateIcon(FAppStyle::GetAppStyleSetName(), "LevelEditor.OpenCinematic"),
+							FUIAction(FExecuteAction::CreateLambda([CineAssembly]()
+								{
+									// Note: Normally, when a cine assembly is opened in the editor, its associated map is loaded during the PrepareToActiveAssets step.
+									// By directly opening the asset editor (Sequencer) for this asset, we can skip that map loading step. 
+									GEditor->GetEditorSubsystem<UAssetEditorSubsystem>()->OpenEditorForAssets({ CineAssembly });
+								})),
+							EUserInterfaceActionType::None
+						);
+					}
 
 					InSection.AddMenuEntry(
 						"OpenInAssetEditor",
 						LOCTEXT("OpenInAssetEditor", "Edit Properties"),
 						LOCTEXT("OpenInAssetEditorTooltip", "Edit Properties"),
 						FSlateIcon(FAppStyle::GetAppStyleSetName(), "Icons.Edit"),
-						FUIAction(FExecuteAction::CreateRaw(this, &FCineAssemblyToolsEditorModule::OpenAssemblyForEdit, CineAssembly)),
+						FUIAction(FExecuteAction::CreateLambda([this, SelectedAssets]()
+							{
+								for (const FAssetData& AssetData : SelectedAssets)
+								{
+									UCineAssembly* CineAssembly = CastChecked<UCineAssembly>(AssetData.GetAsset());
+									OpenAssemblyForEdit(CineAssembly);
+								}
+							})),
 						EUserInterfaceActionType::None
 					);
 				}
