ShotManagement: Add engine analytics to Cinematic Assembly Tools
#rnx

diff --git a/Engine/Plugins/VirtualProduction/CinematicAssemblyTools/Source/CineAssemblyToolsEditor/Private/CineAssemblyFactory.cpp b/Engine/Plugins/VirtualProduction/CinematicAssemblyTools/Source/CineAssemblyToolsEditor/Private/CineAssemblyFactory.cpp
--- a/Engine/Plugins/VirtualProduction/CinematicAssemblyTools/Source/CineAssemblyToolsEditor/Private/CineAssemblyFactory.cpp
+++ b/Engine/Plugins/VirtualProduction/CinematicAssemblyTools/Source/CineAssemblyToolsEditor/Private/CineAssemblyFactory.cpp
@@ -4,6 +4,7 @@
 
 #include "AssetToolsModule.h"
 #include "CineAssemblyNamingTokens.h"
+#include "CineAssemblyToolsAnalytics.h"
 #include "ContentBrowserModule.h"
 #include "Framework/Application/SlateApplication.h"
 #include "IContentBrowserSingleton.h"
@@ -136,6 +137,9 @@ void UCineAssemblyFactory::CreateConfiguredAssembly(UCineAssembly* ConfiguredAss
 
 	ConfiguredAssembly->CreateSubAssemblies();
 
+	// Analytics
+	UE::CineAssemblyToolsAnalytics::RecordEvent_CreateAssembly(ConfiguredAssembly);
+
 	// Refresh the content browser to make any new assets and folders are immediately visible
 	IContentBrowserSingleton& ContentBrowser = FModuleManager::LoadModuleChecked<FContentBrowserModule>("ContentBrowser").Get();
 	ContentBrowser.SetSelectedPaths({ CreateAssetPath }, true);
diff --git a/Engine/Plugins/VirtualProduction/CinematicAssemblyTools/Source/CineAssemblyToolsEditor/Private/CineAssemblySchemaFactory.cpp b/Engine/Plugins/VirtualProduction/CinematicAssemblyTools/Source/CineAssemblyToolsEditor/Private/CineAssemblySchemaFactory.cpp
--- a/Engine/Plugins/VirtualProduction/CinematicAssemblyTools/Source/CineAssemblyToolsEditor/Private/CineAssemblySchemaFactory.cpp
+++ b/Engine/Plugins/VirtualProduction/CinematicAssemblyTools/Source/CineAssemblyToolsEditor/Private/CineAssemblySchemaFactory.cpp
@@ -4,6 +4,7 @@
 
 #include "AssetToolsModule.h"
 #include "CineAssemblySchema.h"
+#include "CineAssemblyToolsAnalytics.h"
 #include "ContentBrowserModule.h"
 #include "Framework/Application/SlateApplication.h"
 #include "IContentBrowserSingleton.h"
@@ -82,6 +83,8 @@ void UCineAssemblySchemaFactory::CreateConfiguredSchema(UCineAssemblySchema* Con
 	ConfiguredSchema->SetFlags(Flags);
 	ConfiguredSchema->ClearFlags(RF_Transient);
 
+	UE::CineAssemblyToolsAnalytics::RecordEvent_CreateAssemblySchema();
+
 	// Refresh the content browser to make any new assets and folders are immediately visible
 	IContentBrowserSingleton& ContentBrowser = FModuleManager::LoadModuleChecked<FContentBrowserModule>("ContentBrowser").Get();
 	ContentBrowser.SetSelectedPaths({ CreateAssetPath }, true);
diff --git a/Engine/Plugins/VirtualProduction/CinematicAssemblyTools/Source/CineAssemblyToolsEditor/Private/CineAssemblyToolsAnalytics.cpp b/Engine/Plugins/VirtualProduction/CinematicAssemblyTools/Source/CineAssemblyToolsEditor/Private/CineAssemblyToolsAnalytics.cpp
new file mode 100644
--- /dev/null
+++ b/Engine/Plugins/VirtualProduction/CinematicAssemblyTools/Source/CineAssemblyToolsEditor/Private/CineAssemblyToolsAnalytics.cpp
@@ -0,0 +1,77 @@
+// Copyright Epic Games, Inc. All Rights Reserved.
+
+#include "CineAssemblyToolsAnalytics.h"
+
+#include "CineAssembly.h"
+#include "CineAssemblySchema.h"
+#include "EngineAnalytics.h"
+#include "Interfaces/IPluginManager.h"
+#include "UObject/Package.h"
+
+namespace UE::CineAssemblyToolsAnalytics
+{
+	void RecordEvent_CreateAssembly(const UCineAssembly* const Assembly)
+	{
+		if (FEngineAnalytics::IsAvailable() && Assembly)
+		{
+			const UCineAssemblySchema* const Schema = Assembly->GetSchema();
+			const bool bHasSchema = (Schema != nullptr);
+			bool bHasUserSchema = false;
+
+			if (bHasSchema && Schema->GetPackage())
+			{
+				// Check if the schema's package comes from the Cinematic Assembly Tools plugin
+				const FString Name = Schema->GetPackage()->GetName();
+				const FString PluginPackageName = IPluginManager::Get().FindPlugin(UE_PLUGIN_NAME)->GetMountedAssetPath();
+
+				bHasUserSchema = !(Name.StartsWith(PluginPackageName, ESearchCase::CaseSensitive));
+			}
+
+			TArray<FAnalyticsEventAttribute> EventAttributes;
+			EventAttributes.Add(FAnalyticsEventAttribute(TEXT("HasSchema"), bHasSchema));
+			EventAttributes.Add(FAnalyticsEventAttribute(TEXT("HasUserSchema"), bHasUserSchema));
+
+			FEngineAnalytics::GetProvider().RecordEvent(TEXT("CinematicAssemblyTools.CreateCineAssembly"), EventAttributes);
+		}
+	}
+
+	void RecordEvent_CreateAssemblySchema()
+	{
+		if (FEngineAnalytics::IsAvailable())
+		{
+			FEngineAnalytics::GetProvider().RecordEvent(TEXT("CinematicAssemblyTools.CreateCineAssemblySchema"));
+		}
+	}
+
+	void RecordEvent_CreateProduction()
+	{
+		if (FEngineAnalytics::IsAvailable())
+		{
+			FEngineAnalytics::GetProvider().RecordEvent(TEXT("CinematicAssemblyTools.CreateProduction"));
+		}
+	}
+
+	void RecordEvent_ProductionAddAssetNaming()
+	{
+		if (FEngineAnalytics::IsAvailable())
+		{
+			FEngineAnalytics::GetProvider().RecordEvent(TEXT("CinematicAssemblyTools.Productions.AddAssetNaming"));
+		}
+	}
+
+	void RecordEvent_ProductionCreateTemplateFolders()
+	{
+		if (FEngineAnalytics::IsAvailable())
+		{
+			FEngineAnalytics::GetProvider().RecordEvent(TEXT("CinematicAssemblyTools.Productions.CreateTemplateFolders"));
+		}
+	}
+
+	void RecordEvent_RecordAssembly()
+	{
+		if (FEngineAnalytics::IsAvailable())
+		{
+			FEngineAnalytics::GetProvider().RecordEvent(TEXT("CinematicAssemblyTools.TakeRecorder.RecordAssembly"));
+		}
+	}
+}
diff --git a/Engine/Plugins/VirtualProduction/CinematicAssemblyTools/Source/CineAssemblyToolsEditor/Private/CineAssemblyToolsAnalytics.h b/Engine/Plugins/VirtualProduction/CinematicAssemblyTools/Source/CineAssemblyToolsEditor/Private/CineAssemblyToolsAnalytics.h
new file mode 100644
--- /dev/null
+++ b/Engine/Plugins/VirtualProduction/CinematicAssemblyTools/Source/CineAssemblyToolsEditor/Private/CineAssemblyToolsAnalytics.h
@@ -0,0 +1,17 @@
+// Copyright Epic Games, Inc. All Rights Reserved.
+
+#pragma once
+
+#include "CoreMinimal.h"
+
+class UCineAssembly;
+
+namespace UE::CineAssemblyToolsAnalytics
+{
+	void RecordEvent_CreateAssembly(const UCineAssembly* const Assembly);
+	void RecordEvent_CreateAssemblySchema();
+	void RecordEvent_CreateProduction();
+	void RecordEvent_ProductionAddAssetNaming();
+	void RecordEvent_ProductionCreateTemplateFolders();
+	void RecordEvent_RecordAssembly();
+};
\ No newline at end of file
diff --git a/Engine/Plugins/VirtualProduction/CinematicAssemblyTools/Source/CineAssemblyToolsEditor/Private/ProductionSettings.cpp b/Engine/Plugins/VirtualProduction/CinematicAssemblyTools/Source/CineAssemblyToolsEditor/Private/ProductionSettings.cpp
--- a/Engine/Plugins/VirtualProduction/CinematicAssemblyTools/Source/CineAssemblyToolsEditor/Private/ProductionSettings.cpp
+++ b/Engine/Plugins/VirtualProduction/CinematicAssemblyTools/Source/CineAssemblyToolsEditor/Private/ProductionSettings.cpp
@@ -5,6 +5,7 @@
 #include "Algo/Find.h"
 #include "Algo/IndexOf.h"
 #include "AssetToolsModule.h"
+#include "CineAssemblyToolsAnalytics.h"
 #include "Engine/Engine.h"
 #include "HAL/FileManager.h"
 #include "HAL/PlatformFileManager.h"
@@ -295,6 +296,8 @@ void UProductionSettings::AddProduction(const FCinematicProduction& ProductionTo
 
 	TryUpdateDefaultConfigFile();
 
+	UE::CineAssemblyToolsAnalytics::RecordEvent_CreateProduction();
+
 	ProductionListChangedDelegate.Broadcast();
 }
 
@@ -590,6 +593,8 @@ void UProductionSettings::AddAssetNaming(FGuid ProductionID, const UClass* Asset
 		}
 
 		TryUpdateDefaultConfigFile();
+
+		UE::CineAssemblyToolsAnalytics::RecordEvent_ProductionAddAssetNaming();
 	}
 }
 
diff --git a/Engine/Plugins/VirtualProduction/CinematicAssemblyTools/Source/CineAssemblyToolsEditor/Private/TakeRecorder/CineAssemblyTakeRecorderIntegration.cpp b/Engine/Plugins/VirtualProduction/CinematicAssemblyTools/Source/CineAssemblyToolsEditor/Private/TakeRecorder/CineAssemblyTakeRecorderIntegration.cpp
--- a/Engine/Plugins/VirtualProduction/CinematicAssemblyTools/Source/CineAssemblyToolsEditor/Private/TakeRecorder/CineAssemblyTakeRecorderIntegration.cpp
+++ b/Engine/Plugins/VirtualProduction/CinematicAssemblyTools/Source/CineAssemblyToolsEditor/Private/TakeRecorder/CineAssemblyTakeRecorderIntegration.cpp
@@ -6,6 +6,7 @@
 #include "CineAssemblyFactory.h"
 #include "CineAssemblySchema.h"
 #include "CineAssemblyTakeRecorderSettings.h"
+#include "CineAssemblyToolsAnalytics.h"
 #include "ProductionSettings.h"
 #include "ITakeRecorderModule.h"
 #include "Recorder/TakeRecorder.h"
@@ -84,6 +85,8 @@ void FCineAssemblyTakeRecorderIntegration::OnRecordingInitialized(UTakeRecorder*
 
 				Assembly->CreateSubAssemblies();
 			}
+
+			UE::CineAssemblyToolsAnalytics::RecordEvent_RecordAssembly();
 		}
 	}
 }
diff --git a/Engine/Plugins/VirtualProduction/CinematicAssemblyTools/Source/CineAssemblyToolsEditor/Private/UI/SFolderHierarchyPanel.cpp b/Engine/Plugins/VirtualProduction/CinematicAssemblyTools/Source/CineAssemblyToolsEditor/Private/UI/SFolderHierarchyPanel.cpp
--- a/Engine/Plugins/VirtualProduction/CinematicAssemblyTools/Source/CineAssemblyToolsEditor/Private/UI/SFolderHierarchyPanel.cpp
+++ b/Engine/Plugins/VirtualProduction/CinematicAssemblyTools/Source/CineAssemblyToolsEditor/Private/UI/SFolderHierarchyPanel.cpp
@@ -4,6 +4,7 @@
 
 #include "AssetRegistry/AssetRegistryModule.h"
 #include "AssetViewUtils.h"
+#include "CineAssemblyToolsAnalytics.h"
 #include "CineAssemblyToolsStyle.h"
 #include "ContentBrowserDataSubsystem.h"
 #include "ContentBrowserItem.h"
@@ -550,6 +551,8 @@ FReply SFolderHierarchyPanel::OnCreateTemplateFolders()
 	TreeView->RequestTreeRefresh();
 	SetExpansionRecursive(RootItem, true);
 
+	UE::CineAssemblyToolsAnalytics::RecordEvent_ProductionCreateTemplateFolders();
+
 	return FReply::Handled();
 }
 
