ShotManagement: Hide subsequence list in the Edit Assembly tab.
Change the Create Assembly text based on the selected schema.
#rnx

diff --git a/Engine/Plugins/VirtualProduction/CinematicAssemblyTools/Source/CineAssemblyToolsEditor/Private/CineAssemblyCustomization.cpp b/Engine/Plugins/VirtualProduction/CinematicAssemblyTools/Source/CineAssemblyToolsEditor/Private/CineAssemblyCustomization.cpp
--- a/Engine/Plugins/VirtualProduction/CinematicAssemblyTools/Source/CineAssemblyToolsEditor/Private/CineAssemblyCustomization.cpp
+++ b/Engine/Plugins/VirtualProduction/CinematicAssemblyTools/Source/CineAssemblyToolsEditor/Private/CineAssemblyCustomization.cpp
@@ -295,6 +295,7 @@ void FCineAssemblyCustomization::CustomizeSubsequenceCategory(IDetailLayoutBuild
 	for (int32 Index = 0; Index < SubAssemblyNames.Num(); ++Index)
 	{
 		SubsequenceCategory.AddCustomRow(FText::GetEmpty())
+			.RowTag(GET_MEMBER_NAME_CHECKED(UCineAssembly, SubAssemblyNames))
 			.NameContent()
 			.HAlign(HAlign_Right)
 			[
diff --git a/Engine/Plugins/VirtualProduction/CinematicAssemblyTools/Source/CineAssemblyToolsEditor/Private/UI/CineAssembly/SCineAssemblyConfigWindow.cpp b/Engine/Plugins/VirtualProduction/CinematicAssemblyTools/Source/CineAssemblyToolsEditor/Private/UI/CineAssembly/SCineAssemblyConfigWindow.cpp
--- a/Engine/Plugins/VirtualProduction/CinematicAssemblyTools/Source/CineAssemblyToolsEditor/Private/UI/CineAssembly/SCineAssemblyConfigWindow.cpp
+++ b/Engine/Plugins/VirtualProduction/CinematicAssemblyTools/Source/CineAssemblyToolsEditor/Private/UI/CineAssembly/SCineAssemblyConfigWindow.cpp
@@ -646,13 +646,12 @@ TSharedRef<SWidget> SCineAssemblyConfigWindow::MakeButtonsPanel()
 			SNew(SHorizontalBox)
 
 			+ SHorizontalBox::Slot()
-				.MinWidth(118.0f)
-				.MaxWidth(118.0f)
+				.AutoWidth()
 				.VAlign(VAlign_Center)
 				.Padding(0.0f, 0.0f, 8.0f, 0.0f)
 				[
 					SNew(SButton)
-						.Text(LOCTEXT("CreateAssetButton", "Create Asset"))
+						.Text(this, &SCineAssemblyConfigWindow::GetCreateButtonText)
 						.ButtonStyle(FAppStyle::Get(), "PrimaryButton")
 						.HAlign(HAlign_Center)
 						.OnClicked(this, &SCineAssemblyConfigWindow::OnCreateAssetClicked)
@@ -671,6 +670,15 @@ TSharedRef<SWidget> SCineAssemblyConfigWindow::MakeButtonsPanel()
 		];
 }
 
+FText SCineAssemblyConfigWindow::GetCreateButtonText() const
+{
+	if (SelectedSchema && !SelectedSchema->SchemaName.IsEmpty())
+	{
+		return FText::Format(LOCTEXT("CreateAssetButtonWithSchema", "Create {0}"), FText::FromString(SelectedSchema->SchemaName));
+	}
+	return LOCTEXT("CreateAssetButton", "Create Assembly");
+}
+
 FReply SCineAssemblyConfigWindow::OnCreateAssetClicked()
 {
 	UCineAssemblyFactory::CreateConfiguredAssembly(CineAssemblyToConfigure.Get(), CreateAssetPath);
@@ -730,12 +738,16 @@ TSharedRef<SWidget> SCineAssemblyEditWidget::BuildUI()
 			];
 	}
 
+	TSharedRef<IDetailsView> DetailsView = UE::CineAssemblyWidgets::Private::MakeDetailsView(CineAssembly);
+	DetailsView->SetIsCustomRowVisibleDelegate(FIsCustomRowVisible::CreateSP(this, &SCineAssemblyEditWidget::IsCustomRowVisible));
+	DetailsView->ForceRefresh();
+
 	TabSwitcher = SNew(SWidgetSwitcher)
 		+ SWidgetSwitcher::Slot()
 		[MakeOverviewWidget()]
 
 		+ SWidgetSwitcher::Slot()
-		[UE::CineAssemblyWidgets::Private::MakeDetailsView(CineAssembly)];
+		[DetailsView];
 
 	return SNew(SBorder)
 		.BorderImage(FCineAssemblyToolsStyle::Get().GetBrush("Borders.PanelNoBorder"))
@@ -774,6 +786,15 @@ TSharedRef<SWidget> SCineAssemblyEditWidget::BuildUI()
 		];
 }
 
+bool SCineAssemblyEditWidget::IsCustomRowVisible(FName RowName, FName ParentName)
+{
+	if (RowName == GET_MEMBER_NAME_CHECKED(UCineAssembly, SubAssemblyNames))
+	{
+		return false;
+	}
+	return true;
+}
+
 void SCineAssemblyEditWidget::FindAssembly(FGuid AssemblyID)
 {
 	const FAssetRegistryModule& AssetRegistryModule = FModuleManager::GetModuleChecked<FAssetRegistryModule>("AssetRegistry");
diff --git a/Engine/Plugins/VirtualProduction/CinematicAssemblyTools/Source/CineAssemblyToolsEditor/Private/UI/CineAssembly/SCineAssemblyConfigWindow.h b/Engine/Plugins/VirtualProduction/CinematicAssemblyTools/Source/CineAssemblyToolsEditor/Private/UI/CineAssembly/SCineAssemblyConfigWindow.h
--- a/Engine/Plugins/VirtualProduction/CinematicAssemblyTools/Source/CineAssemblyToolsEditor/Private/UI/CineAssembly/SCineAssemblyConfigWindow.h
+++ b/Engine/Plugins/VirtualProduction/CinematicAssemblyTools/Source/CineAssemblyToolsEditor/Private/UI/CineAssembly/SCineAssemblyConfigWindow.h
@@ -67,6 +67,9 @@ private:
 	/** Creates the widget to display for the Notes tab */
 	TSharedRef<SWidget> MakeNotesWidget();
 
+	/** Returns the text to display on the Create Asset button, based on the selected schema */
+	FText GetCreateButtonText() const;
+
 	/** Closes the window and indicates that a new asset should be created by the asset factory */
 	FReply OnCreateAssetClicked();
 
@@ -169,6 +172,9 @@ private:
 	/** Creates the widget to display for the Overview tab */
 	TSharedRef<SWidget> MakeOverviewWidget();
 
+	/** Filter used by the Details View to determine which custom rows to show */
+	bool IsCustomRowVisible(FName RowName, FName ParentName);
+
 private:
 	/** Switcher that controls which tab widget is currently visible */
 	TSharedPtr<SWidgetSwitcher> TabSwitcher;
