ShotManagement: Properly initialize new schemas with a default name, and add additional validation on the schema name.
Prevent the user from creating a new schema if they have not placed every subsequence into the folder hierarchy.
#rnx

diff --git a/Engine/Plugins/VirtualProduction/CinematicAssemblyTools/Source/CineAssemblyToolsEditor/Private/CineAssemblySchemaCustomization.cpp b/Engine/Plugins/VirtualProduction/CinematicAssemblyTools/Source/CineAssemblyToolsEditor/Private/CineAssemblySchemaCustomization.cpp
--- a/Engine/Plugins/VirtualProduction/CinematicAssemblyTools/Source/CineAssemblyToolsEditor/Private/CineAssemblySchemaCustomization.cpp
+++ b/Engine/Plugins/VirtualProduction/CinematicAssemblyTools/Source/CineAssemblyToolsEditor/Private/CineAssemblySchemaCustomization.cpp
@@ -2,6 +2,8 @@
 
 #include "CineAssemblySchemaCustomization.h"
 
+#include "Algo/Contains.h"
+#include "AssetRegistry/AssetRegistryModule.h"
 #include "CineAssemblySchema.h"
 #include "DetailLayoutBuilder.h"
 #include "PropertyCustomizationHelpers.h"
@@ -38,6 +40,15 @@ void FCineAssemblySchemaCustomization::CustomizeDetails(IDetailLayoutBuilder& De
 	IDetailPropertyRow& DefaultAssemblyNamePropertyRow = DefaultCategory.AddProperty(DefaultAssemblyNamePropertyHandle);
 
 	// Customize the widget for the SchemaName property to add additional validation on the user input text
+	FString ExistingSchemaName;
+	SchemaNamePropertyHandle->GetValue(ExistingSchemaName);
+
+	// If the key name is not yet set, assign it a unique default key name
+	if (ExistingSchemaName.IsEmpty())
+	{
+		SchemaNamePropertyHandle->SetValue(MakeUniqueSchemaName());
+	}
+
 	SchemaNamePropertyRow.CustomWidget()
 		.NameContent()
 		[
@@ -82,8 +93,44 @@ void FCineAssemblySchemaCustomization::CustomizeDetails(IDetailLayoutBuilder& De
 		];
 }
 
+FString FCineAssemblySchemaCustomization::MakeUniqueSchemaName() const
+{
+	const FString BaseName = TEXT("NewCineAssemblySchema");
+	FString WorkingName = BaseName;
+
+	int32 IntSuffix = 1;
+	while (DoesSchemaExistWithName(WorkingName))
+	{
+		WorkingName = FString::Printf(TEXT("%s%d"), *BaseName, IntSuffix);
+		IntSuffix++;
+	}
+
+	return WorkingName;
+}
+
+bool FCineAssemblySchemaCustomization::DoesSchemaExistWithName(const FString& SchemaName) const
+{
+	const FAssetRegistryModule& AssetRegistryModule = FModuleManager::GetModuleChecked<FAssetRegistryModule>("AssetRegistry");
+	TArray<FAssetData> SchemaAssets;
+	AssetRegistryModule.Get().GetAssetsByClass(UCineAssemblySchema::StaticClass()->GetClassPathName(), SchemaAssets);
+
+	return Algo::ContainsBy(SchemaAssets, *SchemaName, &FAssetData::AssetName);
+}
+
 bool FCineAssemblySchemaCustomization::ValidateSchemaName(const FText& InText, FText& OutErrorMessage) const
 {
+	if (InText.IsEmpty())
+	{
+		OutErrorMessage = LOCTEXT("EmptyNameErrorMessage", "Please provide a name for the schema");
+		return false;
+	}
+
+	if (DoesSchemaExistWithName(InText.ToString()))
+	{
+		OutErrorMessage = LOCTEXT("DuplicateNameErrorMessage", "A schema with that name already exists");
+		return false;
+	}
+
 	return FName::IsValidXName(InText.ToString(), INVALID_OBJECTNAME_CHARACTERS INVALID_LONGPACKAGE_CHARACTERS, &OutErrorMessage);
 }
 
diff --git a/Engine/Plugins/VirtualProduction/CinematicAssemblyTools/Source/CineAssemblyToolsEditor/Private/CineAssemblySchemaCustomization.h b/Engine/Plugins/VirtualProduction/CinematicAssemblyTools/Source/CineAssemblyToolsEditor/Private/CineAssemblySchemaCustomization.h
--- a/Engine/Plugins/VirtualProduction/CinematicAssemblyTools/Source/CineAssemblyToolsEditor/Private/CineAssemblySchemaCustomization.h
+++ b/Engine/Plugins/VirtualProduction/CinematicAssemblyTools/Source/CineAssemblyToolsEditor/Private/CineAssemblySchemaCustomization.h
@@ -20,6 +20,12 @@ public:
 	//~ End IDetailCustomization interface */
 
 private:
+	/** Generate a unique schema name */
+	FString MakeUniqueSchemaName() const;
+
+	/** Returns true if a schema already exists with the input name */
+	bool DoesSchemaExistWithName(const FString& SchemaName) const;
+
 	/** Validates the user input text for the associated properties */
 	bool ValidateSchemaName(const FText& InText, FText& OutErrorMessage) const;
 	bool ValidateDefaultAssemblyName(const FText& InText, FText& OutErrorMessage) const;
diff --git a/Engine/Plugins/VirtualProduction/CinematicAssemblyTools/Source/CineAssemblyToolsEditor/Private/UI/CineAssembly/SCineAssemblySchemaWindow.cpp b/Engine/Plugins/VirtualProduction/CinematicAssemblyTools/Source/CineAssemblyToolsEditor/Private/UI/CineAssembly/SCineAssemblySchemaWindow.cpp
--- a/Engine/Plugins/VirtualProduction/CinematicAssemblyTools/Source/CineAssemblyToolsEditor/Private/UI/CineAssembly/SCineAssemblySchemaWindow.cpp
+++ b/Engine/Plugins/VirtualProduction/CinematicAssemblyTools/Source/CineAssemblyToolsEditor/Private/UI/CineAssembly/SCineAssemblySchemaWindow.cpp
@@ -11,6 +11,7 @@
 #include "Framework/MultiBox/MultiBoxBuilder.h"
 #include "IStructureDetailsView.h"
 #include "Misc/ConfigCacheIni.h"
+#include "Misc/MessageDialog.h"
 #include "Modules/ModuleManager.h"
 #include "NamingTokensEngineSubsystem.h"
 #include "PropertyEditorDelegates.h"
@@ -200,6 +201,33 @@ FString SCineAssemblySchemaWindow::GetSchemaName()
 	return TEXT("CineAssemblySchema");
 }
 
+
+bool SCineAssemblySchemaWindow::DoesSchemaExistWithName(const FString& SchemaName) const
+{
+	const FAssetRegistryModule& AssetRegistryModule = FModuleManager::GetModuleChecked<FAssetRegistryModule>("AssetRegistry");
+	TArray<FAssetData> SchemaAssets;
+	AssetRegistryModule.Get().GetAssetsByClass(UCineAssemblySchema::StaticClass()->GetClassPathName(), SchemaAssets);
+
+	return Algo::ContainsBy(SchemaAssets, *SchemaName, &FAssetData::AssetName);
+}
+
+bool SCineAssemblySchemaWindow::ValidateSchemaName(const FText& InText, FText& OutErrorMessage) const
+{
+	if (InText.IsEmpty())
+	{
+		OutErrorMessage = LOCTEXT("EmptyNameErrorMessage", "Please provide a name for the schema");
+		return false;
+	}
+
+	if (DoesSchemaExistWithName(InText.ToString()))
+	{
+		OutErrorMessage = LOCTEXT("DuplicateNameErrorMessage", "A schema with that name already exists");
+		return false;
+	}
+
+	return FName::IsValidXName(InText.ToString(), INVALID_OBJECTNAME_CHARACTERS INVALID_LONGPACKAGE_CHARACTERS, &OutErrorMessage);
+}
+
 TSharedRef<SWidget> SCineAssemblySchemaWindow::MakeMenuPanel()
 {
 	return SNew(SBorder)
@@ -294,6 +322,7 @@ TSharedRef<SWidget> SCineAssemblySchemaWindow::MakeContentPanel()
 									{ 
 										return FText::FromString(SchemaToConfigure->SchemaName); 
 									})
+								.OnVerifyTextChanged(this, &SCineAssemblySchemaWindow::ValidateSchemaName)
 								.OnTextCommitted_Lambda([this](const FText& InText, ETextCommit::Type InCommitType)
 									{
 										if (SchemaToConfigure->SchemaName != InText.ToString())
@@ -330,7 +359,7 @@ TSharedRef<SWidget> SCineAssemblySchemaWindow::MakeButtonsPanel()
 					.Padding(0.0f, 0.0f, 8.0f, 0.0f)
 					[
 						SNew(SButton)
-							.Text(LOCTEXT("CreateAssetButton", "Create Asset"))
+							.Text(LOCTEXT("CreateAssetButton", "Create Schema"))
 							.ButtonStyle(FAppStyle::Get(), "PrimaryButton")
 							.HAlign(HAlign_Center)
 							.OnClicked(this, &SCineAssemblySchemaWindow::OnCreateAssetClicked)
@@ -1250,6 +1279,16 @@ void SCineAssemblySchemaWindow::OnSchemaPropertiesChanged(const FPropertyChanged
 
 FReply SCineAssemblySchemaWindow::OnCreateAssetClicked()
 {
+	// Prevent the user from finishing creating the schema if there are still assets in the list view that have not been placed in the content hierarchy
+	if (AssetListItems.Num() > 0)
+	{
+		const FText DialogMessage = LOCTEXT("RemainingAssetDialog", "The Content Hierarchy tab contains named subsequences that have not yet been placed in the folder tree. "
+			"Please drag and drop the remaining subsequences, or delete them from the asset list if they are not needed.");
+		FMessageDialog::Open(EAppMsgType::Ok, DialogMessage);
+
+		return FReply::Unhandled();
+	}
+
 	UCineAssemblySchemaFactory::CreateConfiguredSchema(SchemaToConfigure.Get(), CreateAssetPath);
 
 	TSharedPtr<SWindow> ContainingWindow = FSlateApplication::Get().FindWidgetWindow(AsShared());
diff --git a/Engine/Plugins/VirtualProduction/CinematicAssemblyTools/Source/CineAssemblyToolsEditor/Private/UI/CineAssembly/SCineAssemblySchemaWindow.h b/Engine/Plugins/VirtualProduction/CinematicAssemblyTools/Source/CineAssemblyToolsEditor/Private/UI/CineAssembly/SCineAssemblySchemaWindow.h
--- a/Engine/Plugins/VirtualProduction/CinematicAssemblyTools/Source/CineAssemblyToolsEditor/Private/UI/CineAssembly/SCineAssemblySchemaWindow.h
+++ b/Engine/Plugins/VirtualProduction/CinematicAssemblyTools/Source/CineAssemblyToolsEditor/Private/UI/CineAssembly/SCineAssemblySchemaWindow.h
@@ -117,7 +117,6 @@ public:
 	/** Searches the asset registry for a Cine Assembly Schema matching the input ID and updates the UI */
 	void FindSchema(FGuid SchemaID);
 
-
 private:
 	/** Constructs the main UI for the widget */
 	TSharedRef<SWidget> BuildUI();
@@ -143,6 +142,12 @@ private:
 	/** Creates the content for the Details and Data tabs (metadata properties are only shown in the Data tab) */
 	TSharedRef<SWidget> MakeDetailsWidget(bool bShowMetadata);
 
+	/** Returns true if a schema already exists with the input name */
+	bool DoesSchemaExistWithName(const FString& SchemaName) const;
+
+	/** Validates the user input text for the schema name */
+	bool ValidateSchemaName(const FText& InText, FText& OutErrorMessage) const;
+
 	/** Filter used by the Details Views to determine which schema properties to display */
 	bool IsPropertyVisible(const FPropertyAndParent& PropertyAndParent, bool bShowMetadata);
 
