ShotManagement: Update icons for cine assemblies and schemas.
Update icons for naming tokens.
Update plugin schema assets to use new icons.
#rnx

diff --git a/Engine/Plugins/Developer/NamingTokens/Source/NamingTokensUncookedOnly/NamingTokensUncookedOnly.Build.cs b/Engine/Plugins/Developer/NamingTokens/Source/NamingTokensUncookedOnly/NamingTokensUncookedOnly.Build.cs
--- a/Engine/Plugins/Developer/NamingTokens/Source/NamingTokensUncookedOnly/NamingTokensUncookedOnly.Build.cs
+++ b/Engine/Plugins/Developer/NamingTokens/Source/NamingTokensUncookedOnly/NamingTokensUncookedOnly.Build.cs
@@ -25,6 +25,7 @@ public class NamingTokensUncookedOnly : ModuleRules
                 "Engine",
                 "EngineAssetDefinitions",
                 "InputCore",
+				"Projects",
                 "Slate",
                 "SlateCore",
                 "UnrealEd",
diff --git a/Engine/Plugins/Developer/NamingTokens/Source/NamingTokensUncookedOnly/Private/AssetDefinition_NamingTokens.cpp b/Engine/Plugins/Developer/NamingTokens/Source/NamingTokensUncookedOnly/Private/AssetDefinition_NamingTokens.cpp
--- a/Engine/Plugins/Developer/NamingTokens/Source/NamingTokensUncookedOnly/Private/AssetDefinition_NamingTokens.cpp
+++ b/Engine/Plugins/Developer/NamingTokens/Source/NamingTokensUncookedOnly/Private/AssetDefinition_NamingTokens.cpp
@@ -3,6 +3,7 @@
 #include "AssetDefinition_NamingTokens.h"
 
 #include "NamingTokens.h"
+#include "NamingTokensStyle.h"
 
 TSoftClassPtr<UObject> UAssetDefinition_NamingTokens::GetAssetClass() const
 {
@@ -19,3 +20,8 @@ FText UAssetDefinition_NamingTokens::GetAssetDisplayName() const
 {
 	return NSLOCTEXT("AssetDefinition_NamingTokens", "AssetDefinition_NamingTokens_DisplayName", "Naming Tokens");
 }
+
+const FSlateBrush* UAssetDefinition_NamingTokens::GetThumbnailBrush(const FAssetData& InAssetData, const FName InClassName) const
+{
+	return FNamingTokensStyle::Get().GetBrush("ClassThumbnail.NamingTokens");
+}
diff --git a/Engine/Plugins/Developer/NamingTokens/Source/NamingTokensUncookedOnly/Private/AssetDefinition_NamingTokens.h b/Engine/Plugins/Developer/NamingTokens/Source/NamingTokensUncookedOnly/Private/AssetDefinition_NamingTokens.h
--- a/Engine/Plugins/Developer/NamingTokens/Source/NamingTokensUncookedOnly/Private/AssetDefinition_NamingTokens.h
+++ b/Engine/Plugins/Developer/NamingTokens/Source/NamingTokensUncookedOnly/Private/AssetDefinition_NamingTokens.h
@@ -15,4 +15,5 @@ public:
 	virtual TSoftClassPtr<UObject> GetAssetClass() const override;
 	virtual TConstArrayView<FAssetCategoryPath> GetAssetCategories() const override;
 	virtual FText GetAssetDisplayName() const override;
+	virtual const FSlateBrush* GetThumbnailBrush(const FAssetData& InAssetData, const FName InClassName) const override;
 };
diff --git a/Engine/Plugins/Developer/NamingTokens/Source/NamingTokensUncookedOnly/Private/NamingTokensStyle.cpp b/Engine/Plugins/Developer/NamingTokens/Source/NamingTokensUncookedOnly/Private/NamingTokensStyle.cpp
new file mode 100644
--- /dev/null
+++ b/Engine/Plugins/Developer/NamingTokens/Source/NamingTokensUncookedOnly/Private/NamingTokensStyle.cpp
@@ -0,0 +1,37 @@
+// Copyright Epic Games, Inc. All Rights Reserved.
+
+#include "NamingTokensStyle.h"
+
+#include "Interfaces/IPluginManager.h"
+#include "Styling/CoreStyle.h"
+#include "Styling/SlateStyleMacros.h"
+#include "Styling/SlateTypes.h"
+
+FName FNamingTokensStyle::StyleName("NamingTokensStyle");
+
+FNamingTokensStyle& FNamingTokensStyle::Get()
+{
+	static FNamingTokensStyle Inst;
+	return Inst;
+}
+
+FNamingTokensStyle::FNamingTokensStyle()
+	: FSlateStyleSet(StyleName)
+{
+	// Icons
+	const FVector2D Icon16x16(16.0f, 16.0f);
+	const FVector2D Icon64x64(64.0f, 64.0f);
+
+	const FString ContentDir = IPluginManager::Get().FindPlugin(TEXT("NamingTokens"))->GetContentDir();
+	SetContentRoot(ContentDir);
+
+	Set("ClassIcon.NamingTokens", new FSlateVectorImageBrush(RootToContentDir(TEXT("Slate/Token_16.svg")), Icon16x16));
+	Set("ClassThumbnail.NamingTokens", new FSlateVectorImageBrush(RootToContentDir(TEXT("Slate/Token_64.svg")), Icon64x64));
+
+	FSlateStyleRegistry::RegisterSlateStyle(*this);
+}
+
+FNamingTokensStyle::~FNamingTokensStyle()
+{
+	FSlateStyleRegistry::UnRegisterSlateStyle(*this);
+}
diff --git a/Engine/Plugins/Developer/NamingTokens/Source/NamingTokensUncookedOnly/Private/NamingTokensStyle.h b/Engine/Plugins/Developer/NamingTokens/Source/NamingTokensUncookedOnly/Private/NamingTokensStyle.h
new file mode 100644
--- /dev/null
+++ b/Engine/Plugins/Developer/NamingTokens/Source/NamingTokensUncookedOnly/Private/NamingTokensStyle.h
@@ -0,0 +1,20 @@
+// Copyright Epic Games, Inc. All Rights Reserved.
+
+#pragma once
+
+#include "Styling/SlateStyle.h"
+#include "Styling/SlateStyleRegistry.h"
+
+/** Slate style set for Naming Tokens */
+class FNamingTokensStyle : public FSlateStyleSet
+{
+public:
+	static FName StyleName;
+
+	/** Access the singleton instance for this style set */
+	static FNamingTokensStyle& Get();
+
+private:
+	FNamingTokensStyle();
+	~FNamingTokensStyle();
+};
diff --git a/Engine/Plugins/Developer/NamingTokens/Source/NamingTokensUncookedOnly/Private/NamingTokensUncookedOnlyModule.cpp b/Engine/Plugins/Developer/NamingTokens/Source/NamingTokensUncookedOnly/Private/NamingTokensUncookedOnlyModule.cpp
--- a/Engine/Plugins/Developer/NamingTokens/Source/NamingTokensUncookedOnly/Private/NamingTokensUncookedOnlyModule.cpp
+++ b/Engine/Plugins/Developer/NamingTokens/Source/NamingTokensUncookedOnly/Private/NamingTokensUncookedOnlyModule.cpp
@@ -5,6 +5,7 @@
 #include "Customization/NamingTokensCustomization.h"
 #include "NamingTokenData.h"
 #include "NamingTokens.h"
+#include "NamingTokensStyle.h"
 
 #include "PropertyEditorModule.h"
 
@@ -19,6 +20,8 @@ void FNamingTokensUncookedOnlyModule::StartupModule()
 
 	PropertyModule.RegisterCustomPropertyTypeLayout(FNamingTokenData::StaticStruct()->GetFName(),
 		FOnGetPropertyTypeCustomizationInstance::CreateStatic(&FNamingTokensDataCustomization::MakeInstance));
+
+	FNamingTokensStyle::Get();
 }
 
 void FNamingTokensUncookedOnlyModule::ShutdownModule()
diff --git a/Engine/Plugins/VirtualProduction/CinematicAssemblyTools/Source/CineAssemblyTools/Private/CineAssemblySchema.cpp b/Engine/Plugins/VirtualProduction/CinematicAssemblyTools/Source/CineAssemblyTools/Private/CineAssemblySchema.cpp
--- a/Engine/Plugins/VirtualProduction/CinematicAssemblyTools/Source/CineAssemblyTools/Private/CineAssemblySchema.cpp
+++ b/Engine/Plugins/VirtualProduction/CinematicAssemblyTools/Source/CineAssemblyTools/Private/CineAssemblySchema.cpp
@@ -10,7 +10,7 @@
 #include "AssetToolsModule.h"
 #endif
 
-const FSoftObjectPath UCineAssemblySchema::DefaultThumbnailPath = FSoftObjectPath(TEXT("/CinematicAssemblyTools/DefaultSchemaThumbnail.DefaultSchemaThumbnail"));
+const FSoftObjectPath UCineAssemblySchema::DefaultThumbnailPath = FSoftObjectPath(TEXT("/CinematicAssemblyTools/Resources/DefaultSchemaThumbnail.DefaultSchemaThumbnail"));
 const FName UCineAssemblySchema::SchemaGuidPropertyName = GET_MEMBER_NAME_CHECKED(UCineAssemblySchema, SchemaGuid);
 
 UCineAssemblySchema::UCineAssemblySchema()
diff --git a/Engine/Plugins/VirtualProduction/CinematicAssemblyTools/Source/CineAssemblyToolsEditor/Private/AssetDefinition_CineAssembly.cpp b/Engine/Plugins/VirtualProduction/CinematicAssemblyTools/Source/CineAssemblyToolsEditor/Private/AssetDefinition_CineAssembly.cpp
--- a/Engine/Plugins/VirtualProduction/CinematicAssemblyTools/Source/CineAssemblyToolsEditor/Private/AssetDefinition_CineAssembly.cpp
+++ b/Engine/Plugins/VirtualProduction/CinematicAssemblyTools/Source/CineAssemblyToolsEditor/Private/AssetDefinition_CineAssembly.cpp
@@ -2,6 +2,7 @@
 
 #include "AssetDefinition_CineAssembly.h"
 #include "CineAssembly.h"
+#include "CineAssemblyToolsStyle.h"
 
 #include "Editor.h"
 #include "FileHelpers.h"
@@ -62,11 +63,14 @@ const FSlateBrush* UAssetDefinition_CineAssembly::GetThumbnailBrush(const FAsset
 	{
 		if (const UCineAssemblySchema* Schema = CineAssembly->GetSchema())
 		{
-			return Schema->GetThumbnailBrush();
+			if (Schema->ThumbnailImage)
+			{
+				return Schema->GetThumbnailBrush();
+			}
 		}
 	}
 
-	return nullptr;
+	return FCineAssemblyToolsStyle::Get().GetBrush("ClassThumbnail.CineAssembly");
 }
 
 EAssetCommandResult UAssetDefinition_CineAssembly::OpenAssets(const FAssetOpenArgs& OpenArgs) const
diff --git a/Engine/Plugins/VirtualProduction/CinematicAssemblyTools/Source/CineAssemblyToolsEditor/Private/AssetDefinition_CineAssemblySchema.cpp b/Engine/Plugins/VirtualProduction/CinematicAssemblyTools/Source/CineAssemblyToolsEditor/Private/AssetDefinition_CineAssemblySchema.cpp
--- a/Engine/Plugins/VirtualProduction/CinematicAssemblyTools/Source/CineAssemblyToolsEditor/Private/AssetDefinition_CineAssemblySchema.cpp
+++ b/Engine/Plugins/VirtualProduction/CinematicAssemblyTools/Source/CineAssemblyToolsEditor/Private/AssetDefinition_CineAssemblySchema.cpp
@@ -4,6 +4,7 @@
 
 #include "CineAssemblySchema.h"
 #include "CineAssemblyToolsEditorModule.h"
+#include "CineAssemblyToolsStyle.h"
 #include "UI/CineAssembly/SCineAssemblySchemaWindow.h"
 
 #define LOCTEXT_NAMESPACE "AssetTypeActions"
@@ -33,10 +34,13 @@ const FSlateBrush* UAssetDefinition_CineAssemblySchema::GetThumbnailBrush(const
 {
 	if (UCineAssemblySchema* CineAssemblySchema = Cast<UCineAssemblySchema>(InAssetData.GetAsset()))
 	{
-		return CineAssemblySchema->GetThumbnailBrush();
+		if (CineAssemblySchema->ThumbnailImage)
+		{
+			return CineAssemblySchema->GetThumbnailBrush();
+		}
 	}
 
-	return nullptr;
+	return FCineAssemblyToolsStyle::Get().GetBrush("ClassThumbnail.CineAssemblySchema");
 }
 
 FAssetSupportResponse UAssetDefinition_CineAssemblySchema::CanRename(const FAssetData& InAsset) const
diff --git a/Engine/Plugins/VirtualProduction/CinematicAssemblyTools/Source/CineAssemblyToolsEditor/Private/CineAssemblyToolsStyle.cpp b/Engine/Plugins/VirtualProduction/CinematicAssemblyTools/Source/CineAssemblyToolsEditor/Private/CineAssemblyToolsStyle.cpp
--- a/Engine/Plugins/VirtualProduction/CinematicAssemblyTools/Source/CineAssemblyToolsEditor/Private/CineAssemblyToolsStyle.cpp
+++ b/Engine/Plugins/VirtualProduction/CinematicAssemblyTools/Source/CineAssemblyToolsEditor/Private/CineAssemblyToolsStyle.cpp
@@ -28,15 +28,20 @@ FCineAssemblyToolsStyle::FCineAssemblyToolsStyle()
 	SetContentRoot(ContentDir);
 	SetCoreContentRoot(FPaths::EngineContentDir() / TEXT("Slate"));
 
+	Set("ClassIcon.CineAssembly", new FSlateVectorImageBrush(RootToContentDir(TEXT("Slate/SMAssembly_16.svg")), Icon16x16));
+	Set("ClassThumbnail.CineAssembly", new FSlateVectorImageBrush(RootToContentDir(TEXT("Slate/SMAssembly_64.svg")), Icon64x64));
+	Set("ClassIcon.CineAssemblySchema", new FSlateVectorImageBrush(RootToContentDir(TEXT("Slate/SMAssemblySchema_16.svg")), Icon16x16));
+	Set("ClassThumbnail.CineAssemblySchema", new FSlateVectorImageBrush(RootToContentDir(TEXT("Slate/SMAssemblySchema_64.svg")), Icon64x64));
+
 	Set("Thumbnails.Schema", new FSlateVectorImageBrush(RootToContentDir(TEXT("Slate/LevelSequenceActor_64.svg")), Icon64x64));
 
-	Set("Icons.Assembly", new FSlateVectorImageBrush(RootToContentDir(TEXT("Slate/Assembly.svg")), Icon16x16));
-	Set("Icons.Schema", new FSlateVectorImageBrush(RootToContentDir(TEXT("Slate/Schema.svg")), Icon16x16));
+	Set("Icons.Assembly", new FSlateVectorImageBrush(RootToContentDir(TEXT("Slate/SMAssembly_16.svg")), Icon16x16));
+	Set("Icons.Schema", new FSlateVectorImageBrush(RootToContentDir(TEXT("Slate/SMAssemblySchema_16.svg")), Icon16x16));
 
 	Set("Icons.RevisionControl", new FSlateVectorImageBrush(RootToContentDir(TEXT("Slate/RevisionControl.svg")), Icon16x16));
 	Set("Icons.Productions", new FSlateVectorImageBrush(RootToContentDir(TEXT("Slate/Camera.svg")), Icon16x16));
 	Set("Icons.Sequencer", new FSlateVectorImageBrush(RootToContentDir(TEXT("Slate/Sequencer_16.svg")), Icon16x16));
-	Set("Icons.NamingTokens", new FSlateVectorImageBrush(RootToContentDir(TEXT("Slate/NamedSlot.svg")), Icon16x16));
+	Set("Icons.NamingTokens", new FSlateVectorImageBrush(RootToContentDir(TEXT("Slate/SMToken_16.svg")), Icon16x16));
 	Set("Icons.AssetNaming", new FSlateVectorImageBrush(RootToContentDir(TEXT("Slate/Rename.svg")), Icon16x16));
 	Set("Icons.Folder", new FSlateVectorImageBrush(RootToContentDir(TEXT("Slate/folder-closed.svg")), Icon16x16));
 	Set("Icons.Export", new FSlateVectorImageBrush(RootToContentDir(TEXT("Slate/ExportAll.svg")), Icon16x16));
diff --git a/Engine/Source/Editor/EditorWidgets/Private/STemplateStringEditableTextBox.cpp b/Engine/Source/Editor/EditorWidgets/Private/STemplateStringEditableTextBox.cpp
--- a/Engine/Source/Editor/EditorWidgets/Private/STemplateStringEditableTextBox.cpp
+++ b/Engine/Source/Editor/EditorWidgets/Private/STemplateStringEditableTextBox.cpp
@@ -12,6 +12,7 @@
 #include "Styling/SlateTypes.h"
 #include "Styling/StyleColors.h"
 #include "TemplateStringSyntaxHighlighterMarshaller.h"
+#include "Widgets/Images/SImage.h"
 #include "Widgets/Layout/SBox.h"
 #include "Widgets/Notifications/SPopUpErrorText.h"
 #include "Widgets/SBoxPanel.h"
@@ -61,9 +62,8 @@ void STemplateStringEditableTextBox::Construct(const FArguments& InArgs)
 			.AutoWidth()
 			.Padding(4, 0)
 			[
-				SNew(STextBlock)
-				.Text(FText::FromString(TEXT("{ }")))
-				.ColorAndOpacity(FSlateColor(FStyleColors::AccentPurple))
+				SNew(SImage)
+					.Image(FAppStyle::Get().GetBrush("Icons.TokenTextBox"))
 			]
 
 			+ SHorizontalBox::Slot()
diff --git a/Engine/Source/Runtime/SlateCore/Private/Styling/StarshipCoreStyle.cpp b/Engine/Source/Runtime/SlateCore/Private/Styling/StarshipCoreStyle.cpp
--- a/Engine/Source/Runtime/SlateCore/Private/Styling/StarshipCoreStyle.cpp
+++ b/Engine/Source/Runtime/SlateCore/Private/Styling/StarshipCoreStyle.cpp
@@ -422,6 +422,8 @@ TSharedRef<ISlateStyle> FStarshipCoreStyle::Create()
 		Style->Set("Icons.Badge", new IMAGE_BRUSH_SVG("Starship/Common/badge", Icon16x16));
 		Style->Set("Icons.BadgeModified", new IMAGE_BRUSH_SVG("Starship/Common/badge-modified", Icon16x16));
 
+		Style->Set("Icons.TokenTextBox", new IMAGE_BRUSH_SVG("Starship/Common/TokenTextBox_16", Icon16x16));
+
 		// Toolbar Size Icons
 		Style->Set("Icons.Toolbar.Settings", new IMAGE_BRUSH_SVG("Starship/Common/settings", Icon20x20));
 
