ShotManagement: Add missing {activeProduction} token
#rnx

diff --git a/Engine/Plugins/VirtualProduction/CinematicAssemblyTools/Source/CineAssemblyToolsEditor/Private/CineAssemblyToolsEditorModule.cpp b/Engine/Plugins/VirtualProduction/CinematicAssemblyTools/Source/CineAssemblyToolsEditor/Private/CineAssemblyToolsEditorModule.cpp
--- a/Engine/Plugins/VirtualProduction/CinematicAssemblyTools/Source/CineAssemblyToolsEditor/Private/CineAssemblyToolsEditorModule.cpp
+++ b/Engine/Plugins/VirtualProduction/CinematicAssemblyTools/Source/CineAssemblyToolsEditor/Private/CineAssemblyToolsEditorModule.cpp
@@ -8,6 +8,7 @@
 #include "CineAssembly.h"
 #include "CineAssemblyCustomization.h"
 #include "CineAssemblyMetadataCustomization.h"
+#include "CineAssemblyNamingTokens.h"
 #include "CineAssemblySchema.h"
 #include "CineAssemblySchemaCustomization.h"
 #include "CineAssemblyToolsStyle.h"
@@ -97,7 +98,10 @@ void FCineAssemblyToolsEditorModule::OnPostEngineInit()
 		{
 			if (UNamingTokensEngineSubsystem* NamingTokensSubsystem = GEngine->GetEngineSubsystem<UNamingTokensEngineSubsystem>())
 			{
-				FNamingTokenResultData Result = NamingTokensSubsystem->EvaluateTokenString(NameToSanitize);
+				FNamingTokenFilterArgs FilterArgs;
+				FilterArgs.AdditionalNamespacesToInclude.Add(UCineAssemblyNamingTokens::TokenNamespace);
+
+				FNamingTokenResultData Result = NamingTokensSubsystem->EvaluateTokenString(NameToSanitize, FilterArgs);
 				NameToSanitize = Result.EvaluatedText.ToString();
 			}
 		}));
@@ -133,6 +137,8 @@ void FCineAssemblyToolsEditorModule::OnPostEngineInit()
 
 	const FAssetRegistryModule& AssetRegistryModule = FModuleManager::GetModuleChecked<FAssetRegistryModule>("AssetRegistry");
 	AssetRegistryModule.Get().OnInMemoryAssetDeleted().AddRaw(this, &FCineAssemblyToolsEditorModule::OnAssetDeleted);
+
+	RegisterTokens();
 }
 
 void FCineAssemblyToolsEditorModule::ShutdownModule()
@@ -165,6 +171,34 @@ void FCineAssemblyToolsEditorModule::ShutdownModule()
 	SaveOpenTabs(ManagedSchemaTabs);
 }
 
+void FCineAssemblyToolsEditorModule::RegisterTokens()
+{
+	if (UNamingTokensEngineSubsystem* NamingTokensSubsystem = GEngine->GetEngineSubsystem<UNamingTokensEngineSubsystem>())
+	{
+		UCineAssemblyNamingTokens* CineAssemblyNamingTokens = Cast<UCineAssemblyNamingTokens>(NamingTokensSubsystem->GetNamingTokens(UCineAssemblyNamingTokens::TokenNamespace));
+
+		FGuid ExternalTokensGuid = FGuid::NewGuid();
+		TArray<FNamingTokenData>& ExternalTokens = CineAssemblyNamingTokens->RegisterExternalTokens(ExternalTokensGuid);
+
+		// Register the {activeProduction} token with the cine assembly tokens object
+		FNamingTokenData ActiveProductionToken;
+		ActiveProductionToken.TokenKey = TEXT("activeProduction");
+		ActiveProductionToken.DisplayName = LOCTEXT("ActiveProductionTokenName", "Active Production");
+		ActiveProductionToken.TokenProcessorNative.BindLambda([]()
+			{
+				const UProductionSettings* const ProductionSettings = GetDefault<UProductionSettings>();
+				TOptional<const FCinematicProduction> ActiveProduction = ProductionSettings->GetActiveProduction();
+				if (ActiveProduction.IsSet())
+				{
+					return FText::FromString(ActiveProduction.GetValue().ProductionName);
+				}
+				return FText::GetEmpty();
+			});
+
+		ExternalTokens.Add(ActiveProductionToken);
+	}
+}
+
 TSharedRef<SDockTab> FCineAssemblyToolsEditorModule::SpawnProductionWizard(const FSpawnTabArgs& SpawnTabArgs)
 {
 	return SNew(SDockTab)
diff --git a/Engine/Plugins/VirtualProduction/CinematicAssemblyTools/Source/CineAssemblyToolsEditor/Private/CineAssemblyToolsEditorModule.h b/Engine/Plugins/VirtualProduction/CinematicAssemblyTools/Source/CineAssemblyToolsEditor/Private/CineAssemblyToolsEditorModule.h
--- a/Engine/Plugins/VirtualProduction/CinematicAssemblyTools/Source/CineAssemblyToolsEditor/Private/CineAssemblyToolsEditorModule.h
+++ b/Engine/Plugins/VirtualProduction/CinematicAssemblyTools/Source/CineAssemblyToolsEditor/Private/CineAssemblyToolsEditorModule.h
@@ -45,6 +45,9 @@ private:
 	/** Add context menu actions */
 	void RegisterMenus();
 
+	/** Register additional tokens with the cine assembly tokens */
+	void RegisterTokens();
+
 	/** Called whenever a job has its sequence set. */
 	void OnSequenceSet(UMoviePipelineExecutorJob* Job, ULevelSequence* Sequence);
 
diff --git a/Engine/Plugins/VirtualProduction/CinematicAssemblyTools/Source/CineAssemblyToolsEditor/Private/UI/SFolderHierarchyPanel.cpp b/Engine/Plugins/VirtualProduction/CinematicAssemblyTools/Source/CineAssemblyToolsEditor/Private/UI/SFolderHierarchyPanel.cpp
--- a/Engine/Plugins/VirtualProduction/CinematicAssemblyTools/Source/CineAssemblyToolsEditor/Private/UI/SFolderHierarchyPanel.cpp
+++ b/Engine/Plugins/VirtualProduction/CinematicAssemblyTools/Source/CineAssemblyToolsEditor/Private/UI/SFolderHierarchyPanel.cpp
@@ -4,6 +4,7 @@
 
 #include "AssetRegistry/AssetRegistryModule.h"
 #include "AssetViewUtils.h"
+#include "CineAssemblyNamingTokens.h"
 #include "CineAssemblyToolsAnalytics.h"
 #include "CineAssemblyToolsStyle.h"
 #include "ContentBrowserDataSubsystem.h"
@@ -738,7 +739,11 @@ void SFolderHierarchyPanel::CreateFolderFromTemplateRecursive(TSharedPtr<FTempla
 
 		// Evaluate any tokens found in the template path
 		UNamingTokensEngineSubsystem* NamingTokensSubsystem = GEngine->GetEngineSubsystem<UNamingTokensEngineSubsystem>();
-		FNamingTokenResultData Result = NamingTokensSubsystem->EvaluateTokenString(PathToCreate);	
+
+		FNamingTokenFilterArgs FilterArgs;
+		FilterArgs.AdditionalNamespacesToInclude.Add(UCineAssemblyNamingTokens::TokenNamespace);
+
+		FNamingTokenResultData Result = NamingTokensSubsystem->EvaluateTokenString(PathToCreate, FilterArgs);
 		const FString ResolvedPath = Result.EvaluatedText.ToString();
 
 		// Sanitize the name, in case any tokens were not able to be resolved
@@ -906,7 +911,11 @@ bool SFolderHierarchyPanel::DoesPathExist(const FString& Path) const
 {
 	// Evaluate any tokens found in the template path before checking if the folder already exists
 	UNamingTokensEngineSubsystem* NamingTokensSubsystem = GEngine->GetEngineSubsystem<UNamingTokensEngineSubsystem>();
-	const FNamingTokenResultData Result = NamingTokensSubsystem->EvaluateTokenString(Path);
+
+	FNamingTokenFilterArgs FilterArgs;
+	FilterArgs.AdditionalNamespacesToInclude.Add(UCineAssemblyNamingTokens::TokenNamespace);
+
+	const FNamingTokenResultData Result = NamingTokensSubsystem->EvaluateTokenString(Path, FilterArgs);
 	const FString ResolvedPath = Result.EvaluatedText.ToString();
 
 	FAssetRegistryModule& AssetRegistryModule = FModuleManager::GetModuleChecked<FAssetRegistryModule>("AssetRegistry");
