ShotManagement: Fix potential data loss when changing settings in the Production Wizard, due to DefaultEngine.ini not being writable.
When adding a new production, if it the first production in the list, set it as active.
#rnx

diff --git a/Engine/Plugins/VirtualProduction/CinematicAssemblyTools/Source/CineAssemblyToolsEditor/CineAssemblyToolsEditor.build.cs b/Engine/Plugins/VirtualProduction/CinematicAssemblyTools/Source/CineAssemblyToolsEditor/CineAssemblyToolsEditor.build.cs
--- a/Engine/Plugins/VirtualProduction/CinematicAssemblyTools/Source/CineAssemblyToolsEditor/CineAssemblyToolsEditor.build.cs
+++ b/Engine/Plugins/VirtualProduction/CinematicAssemblyTools/Source/CineAssemblyToolsEditor/CineAssemblyToolsEditor.build.cs
@@ -40,6 +40,7 @@ namespace UnrealBuildTool.Rules
 					"NamingTokens",
 					"Projects",
 					"PropertyEditor",
+					"SharedSettingsWidgets",
 					"Slate",
 					"SlateCore",
 					"SourceControl",
diff --git a/Engine/Plugins/VirtualProduction/CinematicAssemblyTools/Source/CineAssemblyToolsEditor/Private/ProductionSettings.cpp b/Engine/Plugins/VirtualProduction/CinematicAssemblyTools/Source/CineAssemblyToolsEditor/Private/ProductionSettings.cpp
--- a/Engine/Plugins/VirtualProduction/CinematicAssemblyTools/Source/CineAssemblyToolsEditor/Private/ProductionSettings.cpp
+++ b/Engine/Plugins/VirtualProduction/CinematicAssemblyTools/Source/CineAssemblyToolsEditor/Private/ProductionSettings.cpp
@@ -16,6 +16,7 @@
 #include "NamingTokensEngineSubsystem.h"
 #include "Sections/MovieSceneSubSection.h"
 #include "SourceControlHelpers.h"
+#include "SSettingsEditorCheckoutNotice.h"
 
 #define LOCTEXT_NAMESPACE "ProductionSettings"
 
@@ -294,7 +295,12 @@ void UProductionSettings::AddProduction(const FCinematicProduction& ProductionTo
 {
 	Productions.Add(ProductionToAdd);
 
-	TryUpdateDefaultConfigFile();
+	if (Productions.Num() == 1)
+	{
+		SetActiveProduction(ProductionToAdd.ProductionID);
+	}
+
+	UpdateConfig();
 
 	UE::CineAssemblyToolsAnalytics::RecordEvent_CreateProduction();
 
@@ -313,12 +319,31 @@ void UProductionSettings::DuplicateProduction(FGuid ProductionID)
 
 		Productions.Add(DuplicateProduction);
 
-		TryUpdateDefaultConfigFile();
+		UpdateConfig();
 
 		ProductionListChangedDelegate.Broadcast();
 	}
 }
 
+void UProductionSettings::UpdateConfig()
+{
+	// Try to update the default config file. If unsuccessful, try to make the file writable and try again
+	bool bUpdateResult = TryUpdateDefaultConfigFile();
+	if (!bUpdateResult)
+	{
+		const FString ConfigFilePath = GetDefaultConfigFilename();
+
+		bool bMakeWritableResult = SettingsHelpers::CheckOutOrAddFile(ConfigFilePath);
+		if (!bMakeWritableResult)
+		{
+			bMakeWritableResult = SettingsHelpers::MakeWritable(ConfigFilePath);
+		}
+
+		// Try again to update the config file
+		TryUpdateDefaultConfigFile();
+	}
+}
+
 void UProductionSettings::DeleteProduction(FGuid ProductionID)
 {
 	if (!ProductionID.IsValid())
@@ -340,7 +365,7 @@ void UProductionSettings::DeleteProduction(FGuid ProductionID)
 			SetActiveProduction(ProductionID);
 		}
 
-		TryUpdateDefaultConfigFile();
+		UpdateConfig();
 
 		ProductionListChangedDelegate.Broadcast();
 	}
@@ -352,7 +377,7 @@ void UProductionSettings::RenameProduction(FGuid ProductionID, FString NewName)
 	{
 		Production->ProductionName = NewName;
 		SetActiveProductionName();
-		TryUpdateDefaultConfigFile();
+		UpdateConfig();
 	}
 }
 
@@ -503,7 +528,7 @@ void UProductionSettings::SetDisplayRate(FGuid ProductionID, FFrameRate DisplayR
 			OverrideDefaultStartTime();
 		}
 
-		TryUpdateDefaultConfigFile();
+		UpdateConfig();
 	}
 }
 
@@ -518,7 +543,7 @@ void UProductionSettings::SetStartFrame(FGuid ProductionID, int32 StartFrame)
 			OverrideDefaultStartTime();
 		}
 
-		TryUpdateDefaultConfigFile();
+		UpdateConfig();
 	}
 }
 
@@ -533,7 +558,7 @@ void UProductionSettings::SetSubsequencePriority(FGuid ProductionID, ESubsequenc
 			OverrideSubsequenceHierarchicalBias();
 		}
 
-		TryUpdateDefaultConfigFile();
+		UpdateConfig();
 	}
 }
 
@@ -542,7 +567,7 @@ void UProductionSettings::AddNamespaceToDenyList(FGuid ProductionID, const FStri
 	if (FCinematicProduction* Production = Algo::FindBy(Productions, ProductionID, &FCinematicProduction::ProductionID))
 	{
 		Production->NamingTokenNamespaceDenyList.Add(Namespace);
-		TryUpdateDefaultConfigFile();
+		UpdateConfig();
 	}
 }
 
@@ -551,7 +576,7 @@ void UProductionSettings::RemoveNamespaceFromDenyList(FGuid ProductionID, const
 	if (FCinematicProduction* Production = Algo::FindBy(Productions, ProductionID, &FCinematicProduction::ProductionID))
 	{
 		Production->NamingTokenNamespaceDenyList.Remove(Namespace);
-		TryUpdateDefaultConfigFile();
+		UpdateConfig();
 	}
 }
 
@@ -592,7 +617,7 @@ void UProductionSettings::AddAssetNaming(FGuid ProductionID, const UClass* Asset
 			AssetTools.RegisterDefaultAssetNameForClass(AssetClass, DefaultName);
 		}
 
-		TryUpdateDefaultConfigFile();
+		UpdateConfig();
 
 		UE::CineAssemblyToolsAnalytics::RecordEvent_ProductionAddAssetNaming();
 	}
@@ -623,7 +648,7 @@ void UProductionSettings::RemoveAssetNaming(FGuid ProductionID, const UClass* As
 			}
 		}
 
-		TryUpdateDefaultConfigFile();
+		UpdateConfig();
 	}
 }
 
@@ -636,7 +661,7 @@ void UProductionSettings::AddTemplateFolder(FGuid ProductionID, const FString& P
 		NewTemplate.bCreateIfMissing = bCreateIfMissing;
 
 		Production->TemplateFolders.Add(MoveTemp(NewTemplate));
-		TryUpdateDefaultConfigFile();
+		UpdateConfig();
 	}
 }
 
@@ -650,7 +675,7 @@ void UProductionSettings::RemoveTemplateFolder(FGuid ProductionID, const FString
 			Production->TemplateFolders.RemoveAt(IndexToDelete);
 		}
 
-		TryUpdateDefaultConfigFile();
+		UpdateConfig();
 	}
 }
 
@@ -660,7 +685,7 @@ void UProductionSettings::SetTemplateFolderHierarchy(FGuid ProductionID, const T
 	{
 		Production->TemplateFolders = TemplateHierarchy;
 
-		TryUpdateDefaultConfigFile();
+		UpdateConfig();
 	}
 }
 
diff --git a/Engine/Plugins/VirtualProduction/CinematicAssemblyTools/Source/CineAssemblyToolsEditor/Private/ProductionSettings.h b/Engine/Plugins/VirtualProduction/CinematicAssemblyTools/Source/CineAssemblyToolsEditor/Private/ProductionSettings.h
--- a/Engine/Plugins/VirtualProduction/CinematicAssemblyTools/Source/CineAssemblyToolsEditor/Private/ProductionSettings.h
+++ b/Engine/Plugins/VirtualProduction/CinematicAssemblyTools/Source/CineAssemblyToolsEditor/Private/ProductionSettings.h
@@ -223,6 +223,9 @@ private:
 	/** Returns a pointer to the active production (if it exists) */
 	const FCinematicProduction* GetActiveProductionPtr() const;
 
+	/** Try to update the default config file (will attempt to make the file writable if needed) */
+	void UpdateConfig();
+
 private:
 	/** Name of the active production */
 	UPROPERTY(VisibleAnywhere, Category = "Default", meta = (DisplayName = "Active Production"))
