ShotManagement: Fix assembly assets being unnecessarily dirtied when clicking into the notes field but not actually making any edits.
Fix new assemblies having no production set by default.
Fix incorrect tooltips.
#rnx

diff --git a/Engine/Plugins/VirtualProduction/CinematicAssemblyTools/Source/CineAssemblyToolsEditor/Private/UI/CineAssembly/SCineAssemblyConfigWindow.cpp b/Engine/Plugins/VirtualProduction/CinematicAssemblyTools/Source/CineAssemblyToolsEditor/Private/UI/CineAssembly/SCineAssemblyConfigWindow.cpp
--- a/Engine/Plugins/VirtualProduction/CinematicAssemblyTools/Source/CineAssemblyToolsEditor/Private/UI/CineAssembly/SCineAssemblyConfigWindow.cpp
+++ b/Engine/Plugins/VirtualProduction/CinematicAssemblyTools/Source/CineAssemblyToolsEditor/Private/UI/CineAssembly/SCineAssemblyConfigWindow.cpp
@@ -14,6 +14,7 @@
 #include "Misc/ConfigCacheIni.h"
 #include "Modules/ModuleManager.h"
 #include "ObjectTools.h"
+#include "ProductionSettings.h"
 #include "PropertyEditorModule.h"
 #include "UObject/Package.h"
 #include "Settings/ContentBrowserSettings.h"
@@ -56,6 +57,15 @@ void SCineAssemblyConfigWindow::Construct(const FArguments& InArgs, const FStrin
 	UWorld* CurrentWorld = GEditor->GetEditorWorldContext().World();
 	CineAssemblyToConfigure->Level = FSoftObjectPath(CurrentWorld);
 
+	const UProductionSettings* const ProductionSettings = GetDefault<UProductionSettings>();
+	TOptional<const FCinematicProduction> ActiveProduction = ProductionSettings->GetActiveProduction();
+	
+	if (ActiveProduction.IsSet())
+	{
+		CineAssemblyToConfigure->Production = ActiveProduction.GetValue().ProductionID;
+		CineAssemblyToConfigure->ProductionName = ActiveProduction.GetValue().ProductionName;
+	}
+
 	const FVector2D DefaultWindowSize = FVector2D(1400, 750);
 
 	SWindow::Construct(SWindow::FArguments()
@@ -929,8 +939,11 @@ TSharedRef<SWidget> SCineAssemblyEditWidget::MakeOverviewWidget()
 							})
 						.OnTextCommitted_Lambda([this](const FText& InText, ETextCommit::Type InCommitType)
 							{
-								CineAssembly->Modify();
-								CineAssembly->AssemblyNote = InText.ToString();
+								if (CineAssembly->AssemblyNote != InText.ToString())
+								{
+									CineAssembly->Modify();
+									CineAssembly->AssemblyNote = InText.ToString();
+								}
 							})
 				]
 		];
diff --git a/Engine/Plugins/VirtualProduction/CinematicAssemblyTools/Source/CineAssemblyToolsEditor/Private/UI/SProductionListPanel.cpp b/Engine/Plugins/VirtualProduction/CinematicAssemblyTools/Source/CineAssemblyToolsEditor/Private/UI/SProductionListPanel.cpp
--- a/Engine/Plugins/VirtualProduction/CinematicAssemblyTools/Source/CineAssemblyToolsEditor/Private/UI/SProductionListPanel.cpp
+++ b/Engine/Plugins/VirtualProduction/CinematicAssemblyTools/Source/CineAssemblyToolsEditor/Private/UI/SProductionListPanel.cpp
@@ -134,7 +134,6 @@ void SProductionListPanel::Construct(const FArguments& InArgs)
 												[
 													SNew(STextBlock)
 														.Text(LOCTEXT("ImportProductionButton", "Import Production"))
-														.ToolTipText(LOCTEXT("ImportProductionButtonToolTip", "Not yet implemented"))
 												]
 										]
 								]
@@ -231,6 +230,7 @@ TSharedRef<ITableRow> SProductionListPanel::OnGenerateProductionRow(TSharedPtr<F
 					SNew(SButton)
 						.OnClicked(this, &SProductionListPanel::ExportProduction, InItem)
 						.ButtonStyle(FCineAssemblyToolsStyle::Get(), "ProductionWizard.RecessedButton")
+						.ToolTipText(LOCTEXT("ExportProductionButtonToolTip", "Export Production"))
 						[
 							SNew(SImage).Image(FCineAssemblyToolsStyle::Get().GetBrush("Icons.Export"))
 						]
@@ -243,6 +243,7 @@ TSharedRef<ITableRow> SProductionListPanel::OnGenerateProductionRow(TSharedPtr<F
 					SNew(SButton)
 						.ButtonStyle(FCineAssemblyToolsStyle::Get(), "ProductionWizard.RecessedButton")
 						.OnClicked_Lambda(OnDeleteProduction, InItem->ProductionID)
+						.ToolTipText(LOCTEXT("DeleteProductionButtonToolTip", "Delete Production"))
 						[
 							SNew(SImage).Image(FAppStyle::Get().GetBrush("Icons.Delete"))
 						]
