ShotManagement: Fix an issue where schemas cannot be renamed.
#rnx

diff --git a/Engine/Plugins/VirtualProduction/CinematicAssemblyTools/Source/CineAssemblyTools/Private/CineAssemblySchema.cpp b/Engine/Plugins/VirtualProduction/CinematicAssemblyTools/Source/CineAssemblyTools/Private/CineAssemblySchema.cpp
--- a/Engine/Plugins/VirtualProduction/CinematicAssemblyTools/Source/CineAssemblyTools/Private/CineAssemblySchema.cpp
+++ b/Engine/Plugins/VirtualProduction/CinematicAssemblyTools/Source/CineAssemblyTools/Private/CineAssemblySchema.cpp
@@ -92,8 +92,21 @@ bool UCineAssemblySchema::SupportsRename() const
 }
 
 #if WITH_EDITOR
-void UCineAssemblySchema::RenameAsset()
+void UCineAssemblySchema::RenameAsset(const FString& InNewName)
 {
+	// Early-out if the input name already matches the name of this schema
+	if (SchemaName == InNewName)
+	{
+		return;
+	}
+
+	// If this schema does not yet have a valid package yet (i.e. it is still being configured), then there is no need to use Asset Tools to rename it
+	if (GetPackage() == GetTransientPackage())
+	{
+		SchemaName = InNewName;
+		return;
+	}
+
 	// The default behavior for schema assets is to not allow renaming from Content Browser.
 	// However, this function relies on renaming being supported, so we temporarily enable to do the programmatic rename.
 	bSupportsRename = true;
@@ -105,9 +118,13 @@ void UCineAssemblySchema::RenameAsset()
 	TArray<FAssetRenameData> AssetsAndNames;
 	const bool bSoftReferenceOnly = false;
 	const bool bAlsoRenameLocalizedVariants = true;
-	AssetsAndNames.Emplace(FAssetRenameData(this, PackagePath, SchemaName, bSoftReferenceOnly, bAlsoRenameLocalizedVariants));
+	AssetsAndNames.Emplace(FAssetRenameData(this, PackagePath, InNewName, bSoftReferenceOnly, bAlsoRenameLocalizedVariants));
 
-	AssetTools.RenameAssetsWithDialog(AssetsAndNames);
+	EAssetRenameResult Result = AssetTools.RenameAssetsWithDialog(AssetsAndNames);
+	if (Result != EAssetRenameResult::Failure)
+	{
+		SchemaName = InNewName;
+	}
 
 	bSupportsRename = false;
 }
diff --git a/Engine/Plugins/VirtualProduction/CinematicAssemblyTools/Source/CineAssemblyTools/Public/CineAssemblySchema.h b/Engine/Plugins/VirtualProduction/CinematicAssemblyTools/Source/CineAssemblyTools/Public/CineAssemblySchema.h
--- a/Engine/Plugins/VirtualProduction/CinematicAssemblyTools/Source/CineAssemblyTools/Public/CineAssemblySchema.h
+++ b/Engine/Plugins/VirtualProduction/CinematicAssemblyTools/Source/CineAssemblyTools/Public/CineAssemblySchema.h
@@ -78,7 +78,7 @@ public:
 
 #if WITH_EDITOR
 	/** Renames the underlying schema asset */
-	CINEASSEMBLYTOOLS_API void RenameAsset();
+	CINEASSEMBLYTOOLS_API void RenameAsset(const FString& InNewName);
 #endif
 
 private:
diff --git a/Engine/Plugins/VirtualProduction/CinematicAssemblyTools/Source/CineAssemblyToolsEditor/Private/CineAssemblySchemaCustomization.cpp b/Engine/Plugins/VirtualProduction/CinematicAssemblyTools/Source/CineAssemblyToolsEditor/Private/CineAssemblySchemaCustomization.cpp
--- a/Engine/Plugins/VirtualProduction/CinematicAssemblyTools/Source/CineAssemblyToolsEditor/Private/CineAssemblySchemaCustomization.cpp
+++ b/Engine/Plugins/VirtualProduction/CinematicAssemblyTools/Source/CineAssemblyToolsEditor/Private/CineAssemblySchemaCustomization.cpp
@@ -64,8 +64,7 @@ void FCineAssemblySchemaCustomization::CustomizeDetails(IDetailLayoutBuilder& De
 					})
 				.OnTextCommitted_Lambda([this](const FText& InText, ETextCommit::Type InCommitType)
 					{
-						CustomizedSchema->Modify();
-						CustomizedSchema->SchemaName = InText.ToString();
+						CustomizedSchema->RenameAsset(InText.ToString());
 					})
 				.OnVerifyTextChanged_Raw(this, &FCineAssemblySchemaCustomization::ValidateSchemaName)
 		];
@@ -125,6 +124,12 @@ bool FCineAssemblySchemaCustomization::ValidateSchemaName(const FText& InText, F
 		return false;
 	}
 
+	// It is valid if the input text matches the schema's current name
+	if (CustomizedSchema->SchemaName == InText.ToString())
+	{
+		return true;
+	}
+
 	if (DoesSchemaExistWithName(InText.ToString()))
 	{
 		OutErrorMessage = LOCTEXT("DuplicateNameErrorMessage", "A schema with that name already exists");
diff --git a/Engine/Plugins/VirtualProduction/CinematicAssemblyTools/Source/CineAssemblyToolsEditor/Private/UI/CineAssembly/SCineAssemblySchemaWindow.cpp b/Engine/Plugins/VirtualProduction/CinematicAssemblyTools/Source/CineAssemblyToolsEditor/Private/UI/CineAssembly/SCineAssemblySchemaWindow.cpp
--- a/Engine/Plugins/VirtualProduction/CinematicAssemblyTools/Source/CineAssemblyToolsEditor/Private/UI/CineAssembly/SCineAssemblySchemaWindow.cpp
+++ b/Engine/Plugins/VirtualProduction/CinematicAssemblyTools/Source/CineAssemblyToolsEditor/Private/UI/CineAssembly/SCineAssemblySchemaWindow.cpp
@@ -219,6 +219,12 @@ bool SCineAssemblySchemaWindow::ValidateSchemaName(const FText& InText, FText& O
 		return false;
 	}
 
+	// It is valid if the input text matches the schema's current name
+	if (SchemaToConfigure->SchemaName == InText.ToString())
+	{
+		return true;
+	}
+
 	if (DoesSchemaExistWithName(InText.ToString()))
 	{
 		OutErrorMessage = LOCTEXT("DuplicateNameErrorMessage", "A schema with that name already exists");
@@ -325,16 +331,7 @@ TSharedRef<SWidget> SCineAssemblySchemaWindow::MakeContentPanel()
 								.OnVerifyTextChanged(this, &SCineAssemblySchemaWindow::ValidateSchemaName)
 								.OnTextCommitted_Lambda([this](const FText& InText, ETextCommit::Type InCommitType)
 									{
-										if (SchemaToConfigure->SchemaName != InText.ToString())
-										{
-											SchemaToConfigure->Modify();
-											SchemaToConfigure->SchemaName = InText.ToString();
-
-											if (Mode == ESchemaConfigMode::Edit)
-											{
-												SchemaToConfigure->RenameAsset();
-											}
-										}
+										SchemaToConfigure->RenameAsset(InText.ToString());
 									})
 						]
 				]
@@ -1258,14 +1255,7 @@ void SCineAssemblySchemaWindow::OnSchemaPropertiesChanged(const FPropertyChanged
 {
 	const FName PropertyName = PropertyChangedEvent.GetPropertyName();
 
-	if (PropertyName == GET_MEMBER_NAME_CHECKED(UCineAssemblySchema, SchemaName))
-	{
-		if (Mode == ESchemaConfigMode::Edit)
-		{
-			SchemaToConfigure->RenameAsset();
-		}
-	}
-	else if (PropertyName == GET_MEMBER_NAME_CHECKED(FAssemblyMetadataDesc, Key))
+	if (PropertyName == GET_MEMBER_NAME_CHECKED(FAssemblyMetadataDesc, Key))
 	{
 		UNamingTokensEngineSubsystem* NamingTokensSubsystem = GEngine->GetEngineSubsystem<UNamingTokensEngineSubsystem>();
 		UCineAssemblyNamingTokens* CineAssemblyNamingTokens = Cast<UCineAssemblyNamingTokens>(NamingTokensSubsystem->GetNamingTokens(UCineAssemblyNamingTokens::TokenNamespace));
