Fixed text conflict from LOCTEXT macro used with duplicated key (in CinematicAssemblyTools plugin).

#rnx

diff --git a/Engine/Plugins/VirtualProduction/CinematicAssemblyTools/Source/CineAssemblyToolsEditor/Private/UI/CineAssembly/SCineAssemblySchemaWindow.cpp b/Engine/Plugins/VirtualProduction/CinematicAssemblyTools/Source/CineAssemblyToolsEditor/Private/UI/CineAssembly/SCineAssemblySchemaWindow.cpp
--- a/Engine/Plugins/VirtualProduction/CinematicAssemblyTools/Source/CineAssemblyToolsEditor/Private/UI/CineAssembly/SCineAssemblySchemaWindow.cpp
+++ b/Engine/Plugins/VirtualProduction/CinematicAssemblyTools/Source/CineAssemblyToolsEditor/Private/UI/CineAssembly/SCineAssemblySchemaWindow.cpp
@@ -1195,7 +1195,7 @@ bool SCineAssemblySchemaWindow::ValidateFolderName(const FText& InText, FText& O
 
 	if (FoundItem && FoundItem != TreeItem)
 	{
-		OutErrorMessage = LOCTEXT("DuplicateNameErrorMessage", "A folder already exists at this location with this name");
+		OutErrorMessage = LOCTEXT("DuplicateFolderErrorMessage", "A folder already exists at this location with this name");
 		return false;
 	}
 
