ShotManagement: Fix editor crash caused by bad user input to the CreateAssembly() BP Function
#rnx

diff --git a/Engine/Plugins/VirtualProduction/CinematicAssemblyTools/Source/CineAssemblyToolsEditor/Private/ProductionFunctionLibrary.cpp b/Engine/Plugins/VirtualProduction/CinematicAssemblyTools/Source/CineAssemblyToolsEditor/Private/ProductionFunctionLibrary.cpp
--- a/Engine/Plugins/VirtualProduction/CinematicAssemblyTools/Source/CineAssemblyToolsEditor/Private/ProductionFunctionLibrary.cpp
+++ b/Engine/Plugins/VirtualProduction/CinematicAssemblyTools/Source/CineAssemblyToolsEditor/Private/ProductionFunctionLibrary.cpp
@@ -6,6 +6,7 @@
 #include "CineAssembly.h"
 #include "CineAssemblyFactory.h"
 #include "CineAssemblySchema.h"
+#include "Misc/PackageName.h"
 #include "UObject/Package.h"
 
 const TArray<FCinematicProduction> UProductionFunctionLibrary::GetAllProductions()
@@ -86,6 +87,20 @@ void UProductionFunctionLibrary::RenameProduction(FGuid ProductionID, FString Ne
 
 UCineAssembly* UProductionFunctionLibrary::CreateAssembly(UCineAssemblySchema* Schema, TSoftObjectPtr<UWorld> Level, TSoftObjectPtr<UCineAssembly> ParentAssembly, TMap<FString, FString> Metadata, const FString& Path, const FString& Name, bool bUseDefaultNameFromSchema)
 {
+	// If the supplied path is a virtual path, strip off the "/All" from the beginning before proceeding
+	FString CreatePath = Path;
+	if (CreatePath.StartsWith(TEXT("/All/")))
+	{
+		CreatePath.RemoveFromStart(TEXT("/All"));
+	}
+
+	// Now, make sure that the supplied path is actually valid.
+	if (!FPackageName::IsValidPath(CreatePath))
+	{
+		UE_LOG(LogTemp, Error, TEXT("Could not create Cine Assembly asset because the path '%s' is not valid"), *Path);
+		return nullptr;
+	}
+
 	UCineAssembly* NewAssembly = NewObject<UCineAssembly>(GetTransientPackage(), NAME_None, RF_Transient);
 	NewAssembly->SetSchema(Schema);
 	NewAssembly->SetLevel(Level);
@@ -108,7 +123,7 @@ UCineAssembly* UProductionFunctionLibrary::CreateAssembly(UCineAssemblySchema* S
 		NewAssembly->AssemblyName.Resolved = FText::FromString(Name);
 	}
 
-	UCineAssemblyFactory::CreateConfiguredAssembly(NewAssembly, Path);
+	UCineAssemblyFactory::CreateConfiguredAssembly(NewAssembly, CreatePath);
 
 	return NewAssembly;
 }
