ShotManagement: Refactor Cine Assembly config UI to create a reusable "config panel" that can be used in other windows, such as the upcoming duplication tool.
Also moved the assembly "edit widget" into its own file and renamed it the SCineAssemblyPropertyEditor
#rnx

diff --git a/Engine/Plugins/VirtualProduction/CinematicAssemblyTools/Source/CineAssemblyToolsEditor/Private/CineAssemblyToolsEditorModule.cpp b/Engine/Plugins/VirtualProduction/CinematicAssemblyTools/Source/CineAssemblyToolsEditor/Private/CineAssemblyToolsEditorModule.cpp
--- a/Engine/Plugins/VirtualProduction/CinematicAssemblyTools/Source/CineAssemblyToolsEditor/Private/CineAssemblyToolsEditorModule.cpp
+++ b/Engine/Plugins/VirtualProduction/CinematicAssemblyTools/Source/CineAssemblyToolsEditor/Private/CineAssemblyToolsEditorModule.cpp
@@ -24,7 +24,7 @@
 #include "Subsystems/AssetEditorSubsystem.h"
 #include "ToolMenus.h"
 #include "UI/SProductionWizard.h"
-#include "UI/CineAssembly/SCineAssemblyConfigWindow.h"
+#include "UI/CineAssembly/SCineAssemblyPropertyEditor.h"
 #include "UI/CineAssembly/SCineAssemblySchemaWindow.h"
 #include "WorkspaceMenuStructure.h"
 #include "WorkspaceMenuStructureModule.h"
@@ -210,17 +210,17 @@ TSharedRef<SDockTab> FCineAssemblyToolsEditorModule::SpawnProductionWizard(const
 
 TSharedRef<SDockTab> FCineAssemblyToolsEditorModule::SpawnAssemblyTab(const FSpawnTabArgs& SpawnTabArgs, FGuid AssemblyGuid)
 {
-	TSharedPtr<SCineAssemblyEditWidget> Widget = SNew(SCineAssemblyEditWidget, AssemblyGuid);
+	TSharedPtr<SCineAssemblyPropertyEditor> Widget = SNew(SCineAssemblyPropertyEditor, AssemblyGuid);
 	return SpawnAssemblyTab(SpawnTabArgs, Widget);
 }
 
 TSharedRef<SDockTab> FCineAssemblyToolsEditorModule::SpawnAssemblyTab(const FSpawnTabArgs& SpawnTabArgs, UCineAssembly* Assembly)
 {
-	TSharedPtr<SCineAssemblyEditWidget> Widget = SNew(SCineAssemblyEditWidget, Assembly);
+	TSharedPtr<SCineAssemblyPropertyEditor> Widget = SNew(SCineAssemblyPropertyEditor, Assembly);
 	return SpawnAssemblyTab(SpawnTabArgs, Widget);
 }
 
-TSharedRef<SDockTab> FCineAssemblyToolsEditorModule::SpawnAssemblyTab(const FSpawnTabArgs& SpawnTabArgs, TSharedPtr<SCineAssemblyEditWidget> Widget)
+TSharedRef<SDockTab> FCineAssemblyToolsEditorModule::SpawnAssemblyTab(const FSpawnTabArgs& SpawnTabArgs, TSharedPtr<SCineAssemblyPropertyEditor> Widget)
 {
 	TSharedRef<SDockTab> AssemblyTab = SNew(SDockTab)
 		.TabRole(ETabRole::NomadTab)
diff --git a/Engine/Plugins/VirtualProduction/CinematicAssemblyTools/Source/CineAssemblyToolsEditor/Private/CineAssemblyToolsEditorModule.h b/Engine/Plugins/VirtualProduction/CinematicAssemblyTools/Source/CineAssemblyToolsEditor/Private/CineAssemblyToolsEditorModule.h
--- a/Engine/Plugins/VirtualProduction/CinematicAssemblyTools/Source/CineAssemblyToolsEditor/Private/CineAssemblyToolsEditorModule.h
+++ b/Engine/Plugins/VirtualProduction/CinematicAssemblyTools/Source/CineAssemblyToolsEditor/Private/CineAssemblyToolsEditorModule.h
@@ -9,7 +9,7 @@
 
 class UCineAssembly;
 class UCineAssemblySchema;
-class SCineAssemblyEditWidget;
+class SCineAssemblyPropertyEditor;
 class SCineAssemblySchemaWindow;
 class UMoviePipelineExecutorJob;
 class ULevelSequence;
@@ -35,7 +35,7 @@ private:
 	/** Spawns a nomad tab for an Assembly asset to edit its properties */
 	TSharedRef<SDockTab> SpawnAssemblyTab(const FSpawnTabArgs& SpawnTabArgs, UCineAssembly* Assembly);
 	TSharedRef<SDockTab> SpawnAssemblyTab(const FSpawnTabArgs& SpawnTabArgs, FGuid AssemblyGuid);
-	TSharedRef<SDockTab> SpawnAssemblyTab(const FSpawnTabArgs& SpawnTabArgs, TSharedPtr<SCineAssemblyEditWidget> Widget);
+	TSharedRef<SDockTab> SpawnAssemblyTab(const FSpawnTabArgs& SpawnTabArgs, TSharedPtr<SCineAssemblyPropertyEditor> Widget);
 
 	/** Spawns a nomad tab for an Schema asset to edit its properties */
 	TSharedRef<SDockTab> SpawnSchemaTab(const FSpawnTabArgs& SpawnTabArgs, UCineAssemblySchema* Schema);
diff --git a/Engine/Plugins/VirtualProduction/CinematicAssemblyTools/Source/CineAssemblyToolsEditor/Private/UI/CineAssembly/SCineAssemblyConfigPanel.cpp b/Engine/Plugins/VirtualProduction/CinematicAssemblyTools/Source/CineAssemblyToolsEditor/Private/UI/CineAssembly/SCineAssemblyConfigPanel.cpp
new file mode 100644
--- /dev/null
+++ b/Engine/Plugins/VirtualProduction/CinematicAssemblyTools/Source/CineAssemblyToolsEditor/Private/UI/CineAssembly/SCineAssemblyConfigPanel.cpp
@@ -0,0 +1,456 @@
+// Copyright Epic Games, Inc. All Rights Reserved.
+
+#include "SCineAssemblyConfigPanel.h"
+
+#include "CineAssemblyNamingTokens.h"
+#include "CineAssemblyToolsStyle.h"
+#include "IStructureDetailsView.h"
+#include "Modules/ModuleManager.h"
+#include "PropertyEditorModule.h"
+#include "STemplateStringEditableTextBox.h"
+#include "Widgets/Input/SSegmentedControl.h"
+
+#define LOCTEXT_NAMESPACE "SCineAssemblyConfigPanel"
+
+void SCineAssemblyConfigPanel::Construct(const FArguments& InArgs, UCineAssembly* InAssembly)
+{
+	CineAssemblyToConfigure = InAssembly;
+
+	TabSwitcher = SNew(SWidgetSwitcher)
+		+ SWidgetSwitcher::Slot()
+		[MakeDetailsWidget()]
+
+		+ SWidgetSwitcher::Slot()
+		[MakeHierarchyWidget()]
+
+		+ SWidgetSwitcher::Slot()
+		[MakeNotesWidget()];
+
+		ChildSlot
+		[
+			SNew(SVerticalBox)
+
+			+ SVerticalBox::Slot()
+				.AutoHeight()
+				[
+					SNew(SSegmentedControl<int32>)
+						.Value(0)
+						.OnValueChanged_Lambda([Switcher = TabSwitcher](int32 NewValue)
+							{
+								Switcher->SetActiveWidgetIndex(NewValue);
+							})
+
+					+ SSegmentedControl<int32>::Slot(0)
+						.Text(LOCTEXT("DetailsTab", "Details"))
+						.Icon(FSlateIcon(FAppStyle::GetAppStyleSetName(), "Icons.Details").GetIcon())
+
+					+ SSegmentedControl<int32>::Slot(1)
+						.Text(LOCTEXT("HierarchyTab", "Hierarchy"))
+						.Icon(FSlateIcon(FAppStyle::GetAppStyleSetName(), "Icons.FolderClosed").GetIcon())
+
+					+ SSegmentedControl<int32>::Slot(2)
+						.Text(LOCTEXT("NotesTab", "Notes"))
+						.Icon(FSlateIcon(FCineAssemblyToolsStyle::StyleName, "Icons.Notes").GetIcon())
+				]
+
+			+ SVerticalBox::Slot()
+				.FillHeight(1.0f)
+				.HAlign(HAlign_Fill)
+				[
+					TabSwitcher.ToSharedRef()
+				]
+
+			+ SVerticalBox::Slot()
+				.AutoHeight()
+				[
+					SNew(SBorder)
+						.BorderImage(FCineAssemblyToolsStyle::Get().GetBrush("ProductionWizard.PanelBackground"))
+						.Padding(16.0f)
+						[
+							SNew(SHorizontalBox)
+
+							+ SHorizontalBox::Slot()
+								.AutoWidth()
+								.VAlign(VAlign_Center)
+								.Padding(0.0f, 0.0f, 8.0f, 0.0f)
+								[
+									SNew(STextBlock).Text(LOCTEXT("AssemblyNameField", "Assembly Name"))
+								]
+
+							+ SHorizontalBox::Slot()
+								.FillWidth(1.0f)
+								.VAlign(VAlign_Center)
+								[
+									SNew(STemplateStringEditableTextBox)
+										.Text_Lambda([this]() -> FText
+											{
+												return FText::FromString(CineAssemblyToConfigure->AssemblyName.Template);
+											})
+										.ResolvedText_Lambda([this]()
+											{
+												EvaluateTokenString(CineAssemblyToConfigure->AssemblyName);
+												return CineAssemblyToConfigure->AssemblyName.Resolved;
+											})
+										.OnTextCommitted_Lambda([this](const FText& InText, ETextCommit::Type InCommitType)
+											{
+												CineAssemblyToConfigure->Modify();
+												CineAssemblyToConfigure->AssemblyName.Template = InText.ToString();
+											})
+								]
+						]
+				]
+		];
+}
+
+void SCineAssemblyConfigPanel::Refresh()
+{
+	// The details view needs to be redrawn to show the new metadata fields from the selected schema
+	DetailsView->ForceRefresh();
+
+	// Recreate the hierarchy tree items based on the selected schema
+	PopulateHierarchyTree();
+}
+
+void SCineAssemblyConfigPanel::EvaluateTokenString(FTemplateString& StringToEvaluate)
+{
+	FDateTime CurrentTime = FDateTime::Now();
+	if ((CurrentTime - LastTokenUpdateTime).GetSeconds() >= 1.0f)
+	{
+		StringToEvaluate.Resolved = UCineAssemblyNamingTokens::GetResolvedText(StringToEvaluate.Template, CineAssemblyToConfigure);
+		LastTokenUpdateTime = CurrentTime;
+	}
+}
+
+TSharedRef<SWidget> SCineAssemblyConfigPanel::MakeDetailsWidget()
+{
+	FPropertyEditorModule& PropertyEditor = FModuleManager::Get().LoadModuleChecked<FPropertyEditorModule>(TEXT("PropertyEditor"));
+
+	FDetailsViewArgs DetailsViewArgs;
+	DetailsViewArgs.bAllowSearch = false;
+	DetailsViewArgs.bHideSelectionTip = true;
+
+	DetailsView = PropertyEditor.CreateDetailView(DetailsViewArgs);
+	DetailsView->SetObject(CineAssemblyToConfigure, true);
+
+	return SNew(SVerticalBox)
+
+		+ SVerticalBox::Slot()
+		.AutoHeight()
+		[
+			SNew(SBorder)
+				.BorderImage(FCineAssemblyToolsStyle::Get().GetBrush("Borders.RecessedNoBorder"))
+				.Padding(16.0f)
+				[
+					SNew(SVerticalBox)
+
+					+ SVerticalBox::Slot()
+						.Padding(0.0f, 0.0f, 0.0f, 16.0f)
+						.AutoHeight()
+						[
+							SNew(SHorizontalBox)
+
+							+ SHorizontalBox::Slot()
+								.AutoWidth()
+								.VAlign(VAlign_Center)
+								.Padding(0.0f, 0.0f, 8.0f, 0.0f)
+								[
+									SNew(SImage)
+										.Image_Lambda([this]()
+											{
+												const UCineAssemblySchema* Schema = CineAssemblyToConfigure->GetSchema();
+												return Schema ? Schema->GetThumbnailBrush() : FCineAssemblyToolsStyle::Get().GetBrush("Thumbnails.Schema");
+											})
+								]
+
+							+ SHorizontalBox::Slot()
+								.VAlign(VAlign_Center)
+								[
+									SNew(SVerticalBox)
+
+									+ SVerticalBox::Slot()
+										[
+											SNew(STextBlock)
+												.Text_Lambda([this]()
+													{
+														const UCineAssemblySchema* Schema = CineAssemblyToConfigure->GetSchema();
+														return Schema ? FText::FromString(Schema->SchemaName) : LOCTEXT("NoSchemaName", "No Schema");
+													})
+										]
+
+									+ SVerticalBox::Slot()
+										[
+											SNew(STextBlock)
+												.Text(LOCTEXT("SchemClassName", "Cine Assembly Schema"))
+												.ColorAndOpacity(FSlateColor::UseSubduedForeground())
+										]
+								]
+						]
+
+					+ SVerticalBox::Slot()
+						.AutoHeight()
+						[
+							SNew(STextBlock)
+								.AutoWrapText(true)
+								.Text_Lambda([this]()
+									{
+										if (const UCineAssemblySchema* Schema = CineAssemblyToConfigure->GetSchema())
+										{
+											return !Schema->Description.IsEmpty() ? FText::FromString(Schema->Description) : LOCTEXT("EmptyDescription", "No description");
+										}
+										return LOCTEXT("SchemaInstructions", "Choose a schema to use as the base for configuring your Cine Assembly, or proceed with no schema.");
+									})
+						]
+				]
+		]
+
+		+ SVerticalBox::Slot()
+		.FillContentHeight(1.0f)
+		[
+			DetailsView.ToSharedRef()
+		];
+}
+
+TSharedRef<SWidget> SCineAssemblyConfigPanel::MakeHierarchyWidget()
+{
+	HierarchyTreeView = SNew(STreeView<TSharedPtr<FHierarchyTreeItem>>)
+		.TreeItemsSource(&HierarchyTreeItems)
+		.SelectionMode(ESelectionMode::None)
+		.OnGenerateRow(this, &SCineAssemblyConfigPanel::OnGenerateTreeRow)
+		.OnGetChildren(this, &SCineAssemblyConfigPanel::OnGetChildren);
+
+	// Create the hierarchy tree root
+	RootItem = MakeShared<FHierarchyTreeItem>();
+	RootItem->Type = FHierarchyTreeItem::EItemType::Folder;
+
+	FTemplateString RootTemplate;
+	RootTemplate.Template = TEXT("");
+	RootItem->Path = RootTemplate;
+
+	PopulateHierarchyTree();
+
+	// Register a Slate timer that runs at a set frequency to evaluate all of the tokens in the tree view.
+	// This will automatically be unregistered when this window is destroyed.
+	constexpr float TimerFrequency = 1.0f;
+	RegisterActiveTimer(TimerFrequency, FWidgetActiveTimerDelegate::CreateLambda([this](double InCurrentTime, float InDeltaTime) -> EActiveTimerReturnType
+		{
+			EvaluateHierarchyTokensRecursive(RootItem);
+			HierarchyTreeView->RequestTreeRefresh();
+
+			return EActiveTimerReturnType::Continue;
+		}));
+
+	return SNew(SVerticalBox)
+
+		+ SVerticalBox::Slot()
+		.Padding(16.0f)
+		.AutoHeight()
+		[
+			SNew(STextBlock)
+				.Text(LOCTEXT("HierarchyInstructions", "The following content will be created as defined by the selected Schema."))
+				.AutoWrapText(true)
+		]
+
+		+ SVerticalBox::Slot()
+		[
+			SNew(SBorder)
+				.BorderImage(FCineAssemblyToolsStyle::Get().GetBrush("ProductionWizard.RecessedBackground"))
+				.Padding(8.0f)
+				[
+					HierarchyTreeView.ToSharedRef()
+				]
+		];
+}
+
+void SCineAssemblyConfigPanel::PopulateHierarchyTree()
+{
+	HierarchyTreeItems.Empty();
+	HierarchyTreeItems.Add(RootItem);
+	RootItem->ChildAssets.Empty();
+	RootItem->ChildFolders.Empty();
+
+	auto AddItemsToTree = [this](TArray<FTemplateString>& ItemList, FHierarchyTreeItem::EItemType ItemType)
+		{
+			// Sort the list so that paths are added to the tree in the proper order
+			ItemList.Sort([](const FTemplateString& A, const FTemplateString& B) { return A.Template < B.Template; });
+
+			for (const FTemplateString& ItemName : ItemList)
+			{
+				const FString ParentPath = FPaths::GetPath(ItemName.Template);
+
+				// Walk the tree until we find an item whose path matches the parent path. The new tree item will be created as one of its children
+				TSharedPtr<FHierarchyTreeItem> ParentItem = FindItemAtPathRecursive(RootItem, ParentPath);
+				if (ParentItem)
+				{
+					TSharedPtr<FHierarchyTreeItem> NewItem = MakeShared<FHierarchyTreeItem>();
+					NewItem->Type = ItemType;
+					NewItem->Path = ItemName;
+
+					if (ItemType == FHierarchyTreeItem::EItemType::Folder)
+					{
+						ParentItem->ChildFolders.Add(NewItem);
+					}
+					else
+					{
+						ParentItem->ChildAssets.Add(NewItem);
+					}
+				}
+			}
+		};
+
+	if (const UCineAssemblySchema* Schema = CineAssemblyToConfigure->GetSchema())
+	{
+		TArray<FTemplateString> FolderTemplates;
+		FolderTemplates.Reserve(Schema->FoldersToCreate.Num());
+		Algo::Transform(Schema->FoldersToCreate, FolderTemplates, [](const FString& TemplateString)
+			{
+				FTemplateString FolderTemplate;
+				FolderTemplate.Template = TemplateString;
+				return FolderTemplate;
+			});
+
+		AddItemsToTree(FolderTemplates, FHierarchyTreeItem::EItemType::Folder);
+		AddItemsToTree(CineAssemblyToConfigure->SubAssemblyNames, FHierarchyTreeItem::EItemType::Asset);
+
+		EvaluateHierarchyTokensRecursive(RootItem);
+	}
+
+	HierarchyTreeView->RequestTreeRefresh();
+	ExpandTreeRecursive(RootItem);
+}
+
+void SCineAssemblyConfigPanel::EvaluateHierarchyTokensRecursive(TSharedPtr<FHierarchyTreeItem> TreeItem)
+{
+	// Evaluate the token template string for this tree item
+	TreeItem->Path.Resolved = UCineAssemblyNamingTokens::GetResolvedText(TreeItem->Path.Template, CineAssemblyToConfigure);
+
+	// Evaluate the tokens for all of the child assets, then resort them alphabetically based on the resolved paths
+	for (const TSharedPtr<FHierarchyTreeItem>& Asset : TreeItem->ChildAssets)
+	{
+		EvaluateHierarchyTokensRecursive(Asset);
+	}
+
+	TreeItem->ChildAssets.Sort([](const TSharedPtr<FHierarchyTreeItem>& A, const TSharedPtr<FHierarchyTreeItem>& B) { return A->Path.Resolved.ToString() < B->Path.Resolved.ToString(); });
+
+	// Evaluate the tokens for all of the child folders, then resort them alphabetically based on the resolved paths
+	for (const TSharedPtr<FHierarchyTreeItem>& Child : TreeItem->ChildFolders)
+	{
+		EvaluateHierarchyTokensRecursive(Child);
+	}
+
+	TreeItem->ChildFolders.Sort([](const TSharedPtr<FHierarchyTreeItem>& A, const TSharedPtr<FHierarchyTreeItem>& B) { return A->Path.Resolved.ToString() < B->Path.Resolved.ToString(); });
+}
+
+void SCineAssemblyConfigPanel::ExpandTreeRecursive(TSharedPtr<FHierarchyTreeItem> TreeItem) const
+{
+	if (HierarchyTreeView)
+	{
+		HierarchyTreeView->SetItemExpansion(TreeItem, true);
+	}
+
+	for (const TSharedPtr<FHierarchyTreeItem>& ChildItem : TreeItem->ChildFolders)
+	{
+		ExpandTreeRecursive(ChildItem);
+	}
+}
+
+TSharedPtr<FHierarchyTreeItem> SCineAssemblyConfigPanel::FindItemAtPathRecursive(TSharedPtr<FHierarchyTreeItem> TreeItem, const FString& Path) const
+{
+	if (TreeItem->Path.Template.Equals(Path))
+	{
+		return TreeItem;
+	}
+
+	for (const TSharedPtr<FHierarchyTreeItem>& Child : TreeItem->ChildFolders)
+	{
+		if (const TSharedPtr<FHierarchyTreeItem>& ItemAtPath = FindItemAtPathRecursive(Child, Path))
+		{
+			return ItemAtPath;
+		}
+	}
+
+	return nullptr;
+}
+
+TSharedRef<ITableRow> SCineAssemblyConfigPanel::OnGenerateTreeRow(TSharedPtr<FHierarchyTreeItem> TreeItem, const TSharedRef<STableViewBase>& OwnerTable)
+{
+	TSharedPtr<SImage> Icon = SNew(SImage);
+
+	if (TreeItem->Type == FHierarchyTreeItem::EItemType::Folder)
+	{
+		Icon->SetImage(FCineAssemblyToolsStyle::Get().GetBrush("Icons.Folder"));
+		Icon->SetColorAndOpacity(FAppStyle::Get().GetSlateColor("ContentBrowser.DefaultFolderColor"));
+	}
+	else
+	{
+		Icon->SetImage(FCineAssemblyToolsStyle::Get().GetBrush("Icons.Sequencer"));
+		Icon->SetColorAndOpacity(FLinearColor::White);
+	}
+
+	return SNew(STableRow<TSharedPtr<FHierarchyTreeItem>>, OwnerTable)
+		.Padding(FMargin(8.0f, 2.0f, 8.0f, 0.0f))
+		.Content()
+		[
+			SNew(SHorizontalBox)
+
+			+ SHorizontalBox::Slot()
+				.AutoWidth()
+				.VAlign(VAlign_Center)
+				.Padding(0.0f, 0.0f, 2.0f, 0.0f)
+				[
+					Icon.ToSharedRef()
+				]
+
+			+ SHorizontalBox::Slot()
+				.VAlign(VAlign_Center)
+				[
+					SNew(STextBlock)
+						.Text_Lambda([this, TreeItem]() -> FText
+							{
+								return (TreeItem == RootItem) ? LOCTEXT("RootPathName", "Root Folder") : FText::FromString(FPaths::GetPathLeaf(TreeItem->Path.Resolved.ToString()));
+							})
+				]
+		];
+}
+
+void SCineAssemblyConfigPanel::OnGetChildren(TSharedPtr<FHierarchyTreeItem> TreeItem, TArray<TSharedPtr<FHierarchyTreeItem>>& OutNodes)
+{
+	// Display all of the child assets first, followed by all of the child folders
+	OutNodes.Append(TreeItem->ChildAssets);
+	OutNodes.Append(TreeItem->ChildFolders);
+}
+
+TSharedRef<SWidget> SCineAssemblyConfigPanel::MakeNotesWidget()
+{
+	return SNew(SVerticalBox)
+
+		+ SVerticalBox::Slot()
+		.Padding(16.0f)
+		.AutoHeight()
+		[
+			SNew(STextBlock)
+				.Text(LOCTEXT("NoteInstructions", "The following notes will be saved with the assembly. This can also be edited later."))
+				.AutoWrapText(true)
+		]
+
+		+ SVerticalBox::Slot()
+		[
+			SNew(SBorder)
+				.BorderImage(FCineAssemblyToolsStyle::Get().GetBrush("Borders.Background"))
+				.Padding(16.0f)
+				[
+					SNew(SMultiLineEditableText)
+						.HintText(LOCTEXT("NoteHintText", "Assembly Notes"))
+						.Text_Lambda([this]()
+							{
+								return FText::FromString(CineAssemblyToConfigure->AssemblyNote);
+							})
+						.OnTextCommitted_Lambda([this](const FText& InText, ETextCommit::Type InCommitType)
+							{
+								CineAssemblyToConfigure->Modify();
+								CineAssemblyToConfigure->AssemblyNote = InText.ToString();
+							})
+				]
+		];
+}
+
+#undef LOCTEXT_NAMESPACE
\ No newline at end of file
diff --git a/Engine/Plugins/VirtualProduction/CinematicAssemblyTools/Source/CineAssemblyToolsEditor/Private/UI/CineAssembly/SCineAssemblyConfigPanel.h b/Engine/Plugins/VirtualProduction/CinematicAssemblyTools/Source/CineAssemblyToolsEditor/Private/UI/CineAssembly/SCineAssemblyConfigPanel.h
new file mode 100644
--- /dev/null
+++ b/Engine/Plugins/VirtualProduction/CinematicAssemblyTools/Source/CineAssemblyToolsEditor/Private/UI/CineAssembly/SCineAssemblyConfigPanel.h
@@ -0,0 +1,106 @@
+// Copyright Epic Games, Inc. All Rights Reserved.
+
+#pragma once
+
+#include "Widgets/SCompoundWidget.h"
+
+#include "CineAssembly.h"
+#include "IDetailsView.h"
+#include "Widgets/Layout/SWidgetSwitcher.h"
+#include "Widgets/Views/STreeView.h"
+
+/** An entry in the hierarchy tree view */
+struct FHierarchyTreeItem
+{
+	/** The types of items that can be represented in this tree view */
+	enum class EItemType : uint8
+	{
+		Asset,
+		Folder
+	};
+
+	/** The type of this tree item */
+	EItemType Type;
+
+	/** The relative path of this tree item (possibly containing tokens) */
+	FTemplateString Path;
+
+	/** The children of this item in the tree that are Asset types */
+	TArray<TSharedPtr<FHierarchyTreeItem>> ChildAssets;
+
+	/** The children of this item in the tree that are Folder types */
+	TArray<TSharedPtr<FHierarchyTreeItem>> ChildFolders;
+};
+
+/**
+ * A widget that wraps details and config steps for setting up a cine assembly asset
+ */
+class SCineAssemblyConfigPanel : public SCompoundWidget
+{
+public:
+	SCineAssemblyConfigPanel() = default;
+
+	SLATE_BEGIN_ARGS(SCineAssemblyConfigPanel) {}
+	SLATE_END_ARGS()
+
+	void Construct(const FArguments& InArgs, UCineAssembly* InAssembly);
+
+	/** Refreshes the details view and hierarchy tree view */
+	void Refresh();
+
+private:
+	/** Creates the widget to display for the Overview tab */
+	TSharedRef<SWidget> MakeDetailsWidget();
+
+	/** Creates the widget to display for the Hierarchy tab */
+	TSharedRef<SWidget> MakeHierarchyWidget();
+
+	/** Creates the widget to display for the Notes tab */
+	TSharedRef<SWidget> MakeNotesWidget();
+
+	/** Populate the tree view items from the list of folders and assets specified by the selected schema */
+	void PopulateHierarchyTree();
+
+	/** Generates the row widget for an entry in the tree view */
+	TSharedRef<ITableRow> OnGenerateTreeRow(TSharedPtr<FHierarchyTreeItem> TreeItem, const TSharedRef<STableViewBase>& OwnerTable);
+
+	/** Gets the children of the input tree view item to build additional tree rows */
+	void OnGetChildren(TSharedPtr<FHierarchyTreeItem> TreeItem, TArray<TSharedPtr<FHierarchyTreeItem>>& OutNodes);
+
+	/** Returns the tree item whose path matches the input path */
+	TSharedPtr<FHierarchyTreeItem> FindItemAtPathRecursive(TSharedPtr<FHierarchyTreeItem> TreeItem, const FString& Path) const;
+
+	/** Recursively expands every item in the tree view */
+	void ExpandTreeRecursive(TSharedPtr<FHierarchyTreeItem> TreeItem) const;
+
+	/** Recursively evaluates the tokens in the path of each item in the tree view */
+	void EvaluateHierarchyTokensRecursive(TSharedPtr<FHierarchyTreeItem> TreeItem);
+
+	/**
+	 * Evaluates the input template string with the naming tokens subsystem, and stores the result in the Resolved text.
+	 * This function is throttled to only run at a set frequency, to avoid the potential to constantly query the naming tokens subsystem.
+	 */
+	void EvaluateTokenString(FTemplateString& StringToEvaluate);
+
+private:
+	/** Transient object used only by this UI to configure the properties of the new asset that will get created by the Factory */
+	UCineAssembly* CineAssemblyToConfigure = nullptr;
+
+	/** Switcher that controls which tab widget is currently visible */
+	TSharedPtr<SWidgetSwitcher> TabSwitcher;
+
+	/** Details View displaying the reflected properties of the Cine Assembly being configured */
+	TSharedPtr<IDetailsView> DetailsView;
+
+	/** Items source for the tree view */
+	TArray<TSharedPtr<FHierarchyTreeItem>> HierarchyTreeItems;
+
+	/** The root item in the tree view */
+	TSharedPtr<FHierarchyTreeItem> RootItem;
+
+	/** A read-only tree view of folders and assets that will be created based on the selected schema */
+	TSharedPtr<STreeView<TSharedPtr<FHierarchyTreeItem>>> HierarchyTreeView;
+
+	/** The last time the naming tokens were updated */
+	FDateTime LastTokenUpdateTime;
+};
diff --git a/Engine/Plugins/VirtualProduction/CinematicAssemblyTools/Source/CineAssemblyToolsEditor/Private/UI/CineAssembly/SCineAssemblyConfigWindow.cpp b/Engine/Plugins/VirtualProduction/CinematicAssemblyTools/Source/CineAssemblyToolsEditor/Private/UI/CineAssembly/SCineAssemblyConfigWindow.cpp
--- a/Engine/Plugins/VirtualProduction/CinematicAssemblyTools/Source/CineAssemblyToolsEditor/Private/UI/CineAssembly/SCineAssemblyConfigWindow.cpp
+++ b/Engine/Plugins/VirtualProduction/CinematicAssemblyTools/Source/CineAssemblyToolsEditor/Private/UI/CineAssembly/SCineAssemblyConfigWindow.cpp
@@ -2,51 +2,25 @@
 
 #include "SCineAssemblyConfigWindow.h"
 
-#include "AssetRegistry/AssetRegistryModule.h"
 #include "CineAssemblyFactory.h"
-#include "CineAssemblyNamingTokens.h"
 #include "CineAssemblyToolsStyle.h"
 #include "ContentBrowserModule.h"
 #include "Editor.h"
 #include "IContentBrowserSingleton.h"
 #include "Interfaces/IPluginManager.h"
-#include "IStructureDetailsView.h"
 #include "Misc/ConfigCacheIni.h"
 #include "Modules/ModuleManager.h"
-#include "ObjectTools.h"
 #include "ProductionSettings.h"
-#include "PropertyEditorModule.h"
 #include "UObject/Package.h"
 #include "Settings/ContentBrowserSettings.h"
-#include "STemplateStringEditableTextBox.h"
 #include "Widgets/SBoxPanel.h"
 #include "Widgets/Input/SButton.h"
-#include "Widgets/Input/SSegmentedControl.h"
 #include "Widgets/Layout/SBorder.h"
 #include "Widgets/Layout/SSeparator.h"
 #include "Widgets/Layout/SSplitter.h"
-#include "Widgets/Text/SMultiLineEditableText.h"
-#include "Widgets/Text/STextBlock.h"
 
 #define LOCTEXT_NAMESPACE "SCineAssemblyConfigWindow"
 
-namespace UE::CineAssemblyWidgets::Private
-{
-	TSharedRef<IDetailsView> MakeDetailsView(UCineAssembly* Assembly)
-	{
-		FPropertyEditorModule& PropertyEditor = FModuleManager::Get().LoadModuleChecked<FPropertyEditorModule>(TEXT("PropertyEditor"));
-
-		FDetailsViewArgs DetailsViewArgs;
-		DetailsViewArgs.bAllowSearch = false;
-		DetailsViewArgs.bHideSelectionTip = true;
-
-		TSharedRef<IDetailsView> DetailsView = PropertyEditor.CreateDetailView(DetailsViewArgs);
-		DetailsView->SetObject(Assembly, true);
-
-		return DetailsView;
-	}
-}
-
 void SCineAssemblyConfigWindow::Construct(const FArguments& InArgs, const FString& InCreateAssetPath)
 {
 	CreateAssetPath = InCreateAssetPath;
@@ -89,13 +63,13 @@ void SCineAssemblyConfigWindow::Construct(const FArguments& InArgs, const FStrin
 								+ SSplitter::Slot()
 								.Value(0.7f)
 								[
-									MakeCineTemplatePanel()
+									MakeSchemaPanel()
 								]
 
 								+ SSplitter::Slot()
 								.Value(0.3f)
 								[
-									MakeInfoPanel()
+									SAssignNew(ConfigPanel, SCineAssemblyConfigPanel, CineAssemblyToConfigure.Get())
 								]
 						]
 
@@ -135,7 +109,7 @@ SCineAssemblyConfigWindow::~SCineAssemblyConfigWindow()
 	}
 }
 
-TSharedRef<SWidget> SCineAssemblyConfigWindow::MakeCineTemplatePanel()
+TSharedRef<SWidget> SCineAssemblyConfigWindow::MakeSchemaPanel()
 {
 	FContentBrowserModule& ContentBrowserModule = FModuleManager::Get().LoadModuleChecked<FContentBrowserModule>("ContentBrowser");
 
@@ -208,444 +182,16 @@ void SCineAssemblyConfigWindow::OnSchemaSelected(const FAssetData& AssetData)
 	const FString NoSchemaName = TEXT("NoSchema");
 	if (AssetData.AssetName == NoSchemaName)
 	{
-		SelectedSchema = nullptr;
+		CineAssemblyToConfigure->ChangeSchema(nullptr);
+		ConfigPanel->Refresh();
 	}
 	else if (UCineAssemblySchema* CineAssemblySchema = Cast<UCineAssemblySchema>(AssetData.GetAsset()))
 	{
-		SelectedSchema = CineAssemblySchema;
-	}
-	else
-	{
-		return;
-	}
-
-	CineAssemblyToConfigure->ChangeSchema(SelectedSchema);
-
-	// The details view needs to be redrawn to show the new metadata fields from the selected schema
-	DetailsView->ForceRefresh();
-
-	// Recreate the hierarchy tree items based on the selected schema
-	PopulateHierarchyTree();
-}
-
-TSharedRef<SWidget> SCineAssemblyConfigWindow::MakeInfoPanel()
-{
-	TabSwitcher = SNew(SWidgetSwitcher)
-		+ SWidgetSwitcher::Slot()
-			[ MakeDetailsWidget() ]
-
-		+ SWidgetSwitcher::Slot()
-			[ MakeHierarchyWidget() ]
-
-		+ SWidgetSwitcher::Slot()
-			[ MakeNotesWidget() ];
-
-	return SNew(SVerticalBox)
-
-		+ SVerticalBox::Slot()
-		.AutoHeight()
-		[
-			SNew(SSegmentedControl<int32>)
-				.Value(0)
-				.OnValueChanged_Lambda([Switcher = TabSwitcher](int32 NewValue)
-				{
-					Switcher->SetActiveWidgetIndex(NewValue);
-				})
-
-				+ SSegmentedControl<int32>::Slot(0)
-				.Text(LOCTEXT("DetailsTab", "Details"))
-				.Icon(FSlateIcon(FAppStyle::GetAppStyleSetName(), "Icons.Details").GetIcon())
-
-				+ SSegmentedControl<int32>::Slot(1)
-				.Text(LOCTEXT("HierarchyTab", "Hierarchy"))
-				.Icon(FSlateIcon(FAppStyle::GetAppStyleSetName(), "Icons.FolderClosed").GetIcon())
-
-				+ SSegmentedControl<int32>::Slot(2)
-				.Text(LOCTEXT("NotesTab", "Notes"))
-				.Icon(FSlateIcon(FCineAssemblyToolsStyle::StyleName, "Icons.Notes").GetIcon())
-		]
-
-		+ SVerticalBox::Slot()
-		.FillHeight(1.0f)
-		.HAlign(HAlign_Fill)
-		[
-			TabSwitcher.ToSharedRef()
-		]
-
-		+ SVerticalBox::Slot()
-		.AutoHeight()
-		[
-			SNew(SBorder)
-			.BorderImage(FCineAssemblyToolsStyle::Get().GetBrush("ProductionWizard.PanelBackground"))
-			.Padding(16.0f)
-				[
-					SNew(SHorizontalBox)
-
-					+ SHorizontalBox::Slot()
-						.AutoWidth()
-						.VAlign(VAlign_Center)
-						.Padding(0.0f, 0.0f, 8.0f, 0.0f)
-						[
-							SNew(STextBlock).Text(LOCTEXT("AssemblyNameField", "Assembly Name"))
-						]
-
-					+ SHorizontalBox::Slot()
-						.FillWidth(1.0f)
-						.VAlign(VAlign_Center)
-						[
-							SNew(STemplateStringEditableTextBox)
-								.Text_Lambda([this]() -> FText 
-									{
-										return FText::FromString(CineAssemblyToConfigure->AssemblyName.Template);
-									})
-								.ResolvedText_Lambda([this]() 
-									{ 
-										EvaluateTokenString(CineAssemblyToConfigure->AssemblyName);
-										return CineAssemblyToConfigure->AssemblyName.Resolved;
-									})
-								.OnTextCommitted_Lambda([this](const FText& InText, ETextCommit::Type InCommitType)
-									{
-										CineAssemblyToConfigure->Modify();
-										CineAssemblyToConfigure->AssemblyName.Template = InText.ToString();
-									})
-						]
-				]
-		];
-}
-
-void SCineAssemblyConfigWindow::EvaluateTokenString(FTemplateString& StringToEvaluate)
-{
-	FDateTime CurrentTime = FDateTime::Now();
-	if ((CurrentTime - LastTokenUpdateTime).GetSeconds() >= 1.0f)
-	{
-		StringToEvaluate.Resolved = UCineAssemblyNamingTokens::GetResolvedText(StringToEvaluate.Template, CineAssemblyToConfigure.Get());
-		LastTokenUpdateTime = CurrentTime;
+		CineAssemblyToConfigure->ChangeSchema(CineAssemblySchema);
+		ConfigPanel->Refresh();
 	}
 }
 
-TSharedRef<SWidget> SCineAssemblyConfigWindow::MakeDetailsWidget()
-{
-	DetailsView = UE::CineAssemblyWidgets::Private::MakeDetailsView(CineAssemblyToConfigure.Get());
-
-	return SNew(SVerticalBox)
-
-		+ SVerticalBox::Slot()
-		.AutoHeight()
-		[
-			SNew(SBorder)
-				.BorderImage(FCineAssemblyToolsStyle::Get().GetBrush("Borders.RecessedNoBorder"))
-				.Padding(16.0f)
-				[
-					SNew(SVerticalBox)
-
-					+ SVerticalBox::Slot()
-						.Padding(0.0f, 0.0f, 0.0f, 16.0f)
-						.AutoHeight()
-						[
-							SNew(SHorizontalBox)
-
-							+ SHorizontalBox::Slot()
-								.AutoWidth()
-								.VAlign(VAlign_Center)
-								.Padding(0.0f, 0.0f, 8.0f, 0.0f)
-								[
-									SNew(SImage)
-										.Image_Lambda([this]()
-											{
-												return SelectedSchema ? SelectedSchema->GetThumbnailBrush() : FCineAssemblyToolsStyle::Get().GetBrush("Thumbnails.Schema");
-											})
-								]
-
-								+ SHorizontalBox::Slot()
-								.VAlign(VAlign_Center)
-								[
-									SNew(SVerticalBox)
-
-									+ SVerticalBox::Slot()
-										[
-											SNew(STextBlock)
-												.Text_Lambda([this]()
-													{
-														return SelectedSchema ? FText::FromString(SelectedSchema->SchemaName) : LOCTEXT("NoSchemaName", "No Schema");
-													})
-										]
-
-									+ SVerticalBox::Slot()
-										[
-											SNew(STextBlock)
-												.Text(LOCTEXT("SchemClassName", "Cine Assembly Schema"))
-												.ColorAndOpacity(FSlateColor::UseSubduedForeground())
-										]
-								]
-						]
-
-					+ SVerticalBox::Slot()
-						.AutoHeight()
-						[
-							SNew(STextBlock)
-								.AutoWrapText(true)
-								.Text_Lambda([this]()
-									{
-										if (SelectedSchema)
-										{
-											return !SelectedSchema->Description.IsEmpty() ? FText::FromString(SelectedSchema->Description) : LOCTEXT("EmptyDescription", "No description");
-										}
-										return LOCTEXT("SchemaInstructions", "Choose a schema to use as the base for configuring your Cine Assembly, or proceed with no schema.");
-									})
-						]
-				]
-		]
-
-		+ SVerticalBox::Slot()
-		.FillContentHeight(1.0f)
-		[
-			DetailsView.ToSharedRef()
-		];
-}
-
-TSharedRef<SWidget> SCineAssemblyConfigWindow::MakeHierarchyWidget()
-{
-	HierarchyTreeView = SNew(STreeView<TSharedPtr<FHierarchyTreeItem>>)
-		.TreeItemsSource(&HierarchyTreeItems)
-		.SelectionMode(ESelectionMode::None)
-		.OnGenerateRow(this, &SCineAssemblyConfigWindow::OnGenerateTreeRow)
-		.OnGetChildren(this, &SCineAssemblyConfigWindow::OnGetChildren);
-
-	// Create the hierarchy tree root
-	RootItem = MakeShared<FHierarchyTreeItem>();
-	RootItem->Type = FHierarchyTreeItem::EItemType::Folder;
-
-	FTemplateString RootTemplate;
-	RootTemplate.Template = TEXT("");
-	RootItem->Path = RootTemplate;
-
-	PopulateHierarchyTree();
-
-	// Register a Slate timer that runs at a set frequency to evaluate all of the tokens in the tree view.
-	// This will automatically be unregistered when this window is destroyed.
-	constexpr float TimerFrequency = 1.0f;
-	RegisterActiveTimer(TimerFrequency, FWidgetActiveTimerDelegate::CreateLambda([this](double InCurrentTime, float InDeltaTime) -> EActiveTimerReturnType
-		{
-			EvaluateHierarchyTokensRecursive(RootItem);
-			HierarchyTreeView->RequestTreeRefresh();
-
-			return EActiveTimerReturnType::Continue;
-		}));
-
-	return SNew(SVerticalBox)
-
-		+ SVerticalBox::Slot()
-		.Padding(16.0f)
-		.AutoHeight()
-		[
-			SNew(STextBlock)
-				.Text(LOCTEXT("HierarchyInstructions", "The following content will be created as defined by the selected Schema."))
-				.AutoWrapText(true)
-		]
-
-		+ SVerticalBox::Slot()
-		[
-			SNew(SBorder)
-				.BorderImage(FCineAssemblyToolsStyle::Get().GetBrush("ProductionWizard.RecessedBackground"))
-				.Padding(8.0f)
-				[
-					HierarchyTreeView.ToSharedRef()
-				]
-		];
-}
-
-void SCineAssemblyConfigWindow::PopulateHierarchyTree()
-{
-	HierarchyTreeItems.Empty();
-	HierarchyTreeItems.Add(RootItem);
-	RootItem->ChildAssets.Empty();
-	RootItem->ChildFolders.Empty();
-
-	auto AddItemsToTree = [this](TArray<FTemplateString>& ItemList, FHierarchyTreeItem::EItemType ItemType)
-		{
-			// Sort the list so that paths are added to the tree in the proper order
-			ItemList.Sort([](const FTemplateString& A, const FTemplateString& B) { return A.Template < B.Template; });
-
-			for (const FTemplateString& ItemName : ItemList)
-			{
-				const FString ParentPath = FPaths::GetPath(ItemName.Template);
-
-				// Walk the tree until we find an item whose path matches the parent path. The new tree item will be created as one of its children
-				TSharedPtr<FHierarchyTreeItem> ParentItem = FindItemAtPathRecursive(RootItem, ParentPath);
-				if (ParentItem)
-				{
-					TSharedPtr<FHierarchyTreeItem> NewItem = MakeShared<FHierarchyTreeItem>();
-					NewItem->Type = ItemType;
-					NewItem->Path = ItemName;
-
-					if (ItemType == FHierarchyTreeItem::EItemType::Folder)
-					{
-						ParentItem->ChildFolders.Add(NewItem);
-					}
-					else
-					{
-						ParentItem->ChildAssets.Add(NewItem);
-					}
-				}
-			}
-		};
-
-	if (CineAssemblyToConfigure->BaseSchema)
-	{
-		TArray<FTemplateString> FolderTemplates;
-		FolderTemplates.Reserve(CineAssemblyToConfigure->BaseSchema->FoldersToCreate.Num());
-		Algo::Transform(CineAssemblyToConfigure->BaseSchema->FoldersToCreate, FolderTemplates, [](const FString& TemplateString)
-			{ 
-				FTemplateString FolderTemplate;
-				FolderTemplate.Template = TemplateString;
-				return FolderTemplate;
-			});
-
-		AddItemsToTree(FolderTemplates, FHierarchyTreeItem::EItemType::Folder);
-		AddItemsToTree(CineAssemblyToConfigure->SubAssemblyNames, FHierarchyTreeItem::EItemType::Asset);
-
-		EvaluateHierarchyTokensRecursive(RootItem);
-	}
-
-	HierarchyTreeView->RequestTreeRefresh();
-	ExpandTreeRecursive(RootItem);
-}
-
-void SCineAssemblyConfigWindow::EvaluateHierarchyTokensRecursive(TSharedPtr<FHierarchyTreeItem> TreeItem)
-{
-	// Evaluate the token template string for this tree item
-	TreeItem->Path.Resolved = UCineAssemblyNamingTokens::GetResolvedText(TreeItem->Path.Template, CineAssemblyToConfigure.Get());
-
-	// Evaluate the tokens for all of the child assets, then resort them alphabetically based on the resolved paths
-	for (const TSharedPtr<FHierarchyTreeItem>& Asset : TreeItem->ChildAssets)
-	{
-		EvaluateHierarchyTokensRecursive(Asset);
-	}
-
-	TreeItem->ChildAssets.Sort([](const TSharedPtr<FHierarchyTreeItem>& A, const TSharedPtr<FHierarchyTreeItem>& B) { return A->Path.Resolved.ToString() < B->Path.Resolved.ToString(); });
-
-	// Evaluate the tokens for all of the child folders, then resort them alphabetically based on the resolved paths
-	for (const TSharedPtr<FHierarchyTreeItem>& Child : TreeItem->ChildFolders)
-	{
-		EvaluateHierarchyTokensRecursive(Child);
-	}
-
-	TreeItem->ChildFolders.Sort([](const TSharedPtr<FHierarchyTreeItem>& A, const TSharedPtr<FHierarchyTreeItem>& B) { return A->Path.Resolved.ToString() < B->Path.Resolved.ToString(); });
-}
-
-void SCineAssemblyConfigWindow::ExpandTreeRecursive(TSharedPtr<FHierarchyTreeItem> TreeItem) const
-{
-	if (HierarchyTreeView)
-	{
-		HierarchyTreeView->SetItemExpansion(TreeItem, true);
-	}
-
-	for (const TSharedPtr<FHierarchyTreeItem>& ChildItem : TreeItem->ChildFolders)
-	{
-		ExpandTreeRecursive(ChildItem);
-	}
-}
-
-TSharedPtr<FHierarchyTreeItem> SCineAssemblyConfigWindow::FindItemAtPathRecursive(TSharedPtr<FHierarchyTreeItem> TreeItem, const FString& Path) const
-{
-	if (TreeItem->Path.Template.Equals(Path))
-	{
-		return TreeItem;
-	}
-
-	for (const TSharedPtr<FHierarchyTreeItem>& Child : TreeItem->ChildFolders)
-	{
-		if (const TSharedPtr<FHierarchyTreeItem>& ItemAtPath = FindItemAtPathRecursive(Child, Path))
-		{
-			return ItemAtPath;
-		}
-	}
-
-	return nullptr;
-}
-
-TSharedRef<ITableRow> SCineAssemblyConfigWindow::OnGenerateTreeRow(TSharedPtr<FHierarchyTreeItem> TreeItem, const TSharedRef<STableViewBase>& OwnerTable)
-{
-	TSharedPtr<SImage> Icon = SNew(SImage);
-
-	if (TreeItem->Type == FHierarchyTreeItem::EItemType::Folder)
-	{
-		Icon->SetImage(FCineAssemblyToolsStyle::Get().GetBrush("Icons.Folder"));
-		Icon->SetColorAndOpacity(FAppStyle::Get().GetSlateColor("ContentBrowser.DefaultFolderColor"));
-	}
-	else
-	{
-		Icon->SetImage(FCineAssemblyToolsStyle::Get().GetBrush("Icons.Sequencer"));
-		Icon->SetColorAndOpacity(FLinearColor::White);
-	}
-
-	return SNew(STableRow<TSharedPtr<FHierarchyTreeItem>>, OwnerTable)
-		.Padding(FMargin(8.0f, 2.0f, 8.0f, 0.0f))
-		.Content()
-		[
-			SNew(SHorizontalBox)
-
-			+ SHorizontalBox::Slot()
-				.AutoWidth()
-				.VAlign(VAlign_Center)
-				.Padding(0.0f, 0.0f, 2.0f, 0.0f)
-				[
-					Icon.ToSharedRef()
-				]
-
-			+ SHorizontalBox::Slot()
-				.VAlign(VAlign_Center)
-				[
-					SNew(STextBlock)
-						.Text_Lambda([this, TreeItem]() -> FText
-							{
-								return (TreeItem == RootItem) ? LOCTEXT("RootPathName", "Root Folder") : FText::FromString(FPaths::GetPathLeaf(TreeItem->Path.Resolved.ToString()));
-							})
-				]
-		];
-}
-
-void SCineAssemblyConfigWindow::OnGetChildren(TSharedPtr<FHierarchyTreeItem> TreeItem, TArray<TSharedPtr<FHierarchyTreeItem>>& OutNodes)
-{
-	// Display all of the child assets first, followed by all of the child folders
-	OutNodes.Append(TreeItem->ChildAssets);
-	OutNodes.Append(TreeItem->ChildFolders);
-}
-
-TSharedRef<SWidget> SCineAssemblyConfigWindow::MakeNotesWidget()
-{
-	return SNew(SVerticalBox)
-
-		+ SVerticalBox::Slot()
-		.Padding(16.0f)
-		.AutoHeight()
-		[
-			SNew(STextBlock)
-				.Text(LOCTEXT("NoteInstructions", "The following notes will be saved with the assembly. This can also be edited later."))
-				.AutoWrapText(true)
-		]
-
-		+ SVerticalBox::Slot()
-		[
-			SNew(SBorder)
-				.BorderImage(FCineAssemblyToolsStyle::Get().GetBrush("Borders.Background"))
-				.Padding(16.0f)
-				[
-					SNew(SMultiLineEditableText)
-						.HintText(LOCTEXT("NoteHintText", "Assembly Notes"))
-						.Text_Lambda([this]()
-						{ 
-							return FText::FromString(CineAssemblyToConfigure->AssemblyNote); 
-						})
-						.OnTextCommitted_Lambda([this](const FText& InText, ETextCommit::Type InCommitType)
-						{
-							CineAssemblyToConfigure->Modify();
-							CineAssemblyToConfigure->AssemblyNote = InText.ToString();
-						})
-				]
-		];
-}
-
 TSharedRef<SWidget> SCineAssemblyConfigWindow::MakeButtonsPanel()
 {
 	return SNew(SBorder)
@@ -682,9 +228,12 @@ TSharedRef<SWidget> SCineAssemblyConfigWindow::MakeButtonsPanel()
 
 FText SCineAssemblyConfigWindow::GetCreateButtonText() const
 {
-	if (SelectedSchema && !SelectedSchema->SchemaName.IsEmpty())
+	if (const UCineAssemblySchema* Schema = CineAssemblyToConfigure->GetSchema())
 	{
-		return FText::Format(LOCTEXT("CreateAssetButtonWithSchema", "Create {0}"), FText::FromString(SelectedSchema->SchemaName));
+		if (!Schema->SchemaName.IsEmpty())
+		{
+			return FText::Format(LOCTEXT("CreateAssetButtonWithSchema", "Create {0}"), FText::FromString(Schema->SchemaName));
+		}
 	}
 	return LOCTEXT("CreateAssetButton", "Create Assembly");
 }
@@ -703,250 +252,4 @@ FReply SCineAssemblyConfigWindow::OnCancelClicked()
 	return FReply::Handled();
 }
 
-void SCineAssemblyEditWidget::Construct(const FArguments& InArgs, UCineAssembly* InAssembly)
-{
-	CineAssembly = InAssembly;
-	ChildSlot [ BuildUI() ];
-}
-
-void SCineAssemblyEditWidget::Construct(const FArguments& InArgs, FGuid InAssemblyGuid)
-{
-	// The UI will be temporary because no CineAssembly has been found yet
-	ChildSlot [ BuildUI() ];
-
-	// If the asset registry is still scanning assets, add a callback to find the assembly asset matching the input GUID and update this widget once the scan is finished.
-	// Otherwise, we can find the assembly asset and update the UI immediately. 
-	const FAssetRegistryModule& AssetRegistryModule = FModuleManager::GetModuleChecked<FAssetRegistryModule>("AssetRegistry");
-	if (AssetRegistryModule.Get().IsLoadingAssets())
-	{
-		AssetRegistryModule.Get().OnFilesLoaded().AddSP(this, &SCineAssemblyEditWidget::FindAssembly, InAssemblyGuid);
-	}
-	else
-	{
-		FindAssembly(InAssemblyGuid);
-	}
-}
-
-SCineAssemblyEditWidget::~SCineAssemblyEditWidget()
-{
-	if (FAssetRegistryModule* AssetRegistryModule = FModuleManager::GetModulePtr<FAssetRegistryModule>("AssetRegistry"))
-	{
-		AssetRegistryModule->Get().OnFilesLoaded().RemoveAll(this);
-	}
-}
-
-TSharedRef<SWidget> SCineAssemblyEditWidget::BuildUI()
-{
-	// Build a temporary UI to display while waiting for the assembly to be loaded
-	if (!CineAssembly)
-	{
-		return SNew(SBorder)
-			.BorderImage(FCineAssemblyToolsStyle::Get().GetBrush("Borders.PanelNoBorder"))
-			.Padding(8.0f)
-			[
-				SNew(STextBlock).Text(LOCTEXT("LoadingAssemblyText", "Loading Cine Assembly..."))
-			];
-	}
-
-	TSharedRef<IDetailsView> DetailsView = UE::CineAssemblyWidgets::Private::MakeDetailsView(CineAssembly);
-	DetailsView->SetIsCustomRowVisibleDelegate(FIsCustomRowVisible::CreateSP(this, &SCineAssemblyEditWidget::IsCustomRowVisible));
-	DetailsView->ForceRefresh();
-
-	TabSwitcher = SNew(SWidgetSwitcher)
-		+ SWidgetSwitcher::Slot()
-		[MakeOverviewWidget()]
-
-		+ SWidgetSwitcher::Slot()
-		[DetailsView];
-
-	return SNew(SBorder)
-		.BorderImage(FCineAssemblyToolsStyle::Get().GetBrush("Borders.PanelNoBorder"))
-		.Padding(8.0f)
-		[
-			SNew(SScrollBox)
-				.Orientation(Orient_Vertical)
-
-			+ SScrollBox::Slot()
-				.AutoSize()
-				.HAlign(HAlign_Center)
-				.Padding(0.0f, 0.0f, 0.0f, 8.0f)
-				[
-					SNew(SSegmentedControl<int32>)
-						.Value(0)
-						.OnValueChanged_Lambda([Switcher = TabSwitcher](int32 NewValue)
-							{
-								Switcher->SetActiveWidgetIndex(NewValue);
-							})
-
-						+ SSegmentedControl<int32>::Slot(0)
-						.Text(LOCTEXT("OverviewTab", "Overview"))
-						.Icon(FSlateIcon(FCineAssemblyToolsStyle::StyleName, "Icons.Animation").GetIcon())
-
-						+ SSegmentedControl<int32>::Slot(1)
-						.Text(LOCTEXT("DetailsTab", "Details"))
-						.Icon(FSlateIcon(FAppStyle::GetAppStyleSetName(), "Icons.Details").GetIcon())
-				]
-
-			+ SScrollBox::Slot()
-				.FillSize(1.0f)
-				.HAlign(HAlign_Fill)
-				[
-					TabSwitcher.ToSharedRef()
-				]
-		];
-}
-
-bool SCineAssemblyEditWidget::IsCustomRowVisible(FName RowName, FName ParentName)
-{
-	if (RowName == GET_MEMBER_NAME_CHECKED(UCineAssembly, SubAssemblyNames))
-	{
-		return false;
-	}
-	return true;
-}
-
-void SCineAssemblyEditWidget::FindAssembly(FGuid AssemblyID)
-{
-	const FAssetRegistryModule& AssetRegistryModule = FModuleManager::GetModuleChecked<FAssetRegistryModule>("AssetRegistry");
-
-	// The only search criterion for the asset search is for an asset with an AssemblyID matching the input GUID
-	const TMultiMap<FName, FString> TagValues = { { UCineAssembly::AssemblyGuidPropertyName, AssemblyID.ToString() } };
-
-	TArray<FAssetData> AssemblyAssets;
-	AssetRegistryModule.Get().GetAssetsByTagValues(TagValues, AssemblyAssets);
-
-	// The Assembly ID is unique, so at most one asset should ever be found
-	if (AssemblyAssets.Num() > 0)
-	{
-		CineAssembly = Cast<UCineAssembly>(AssemblyAssets[0].GetAsset());
-
-		// Update the widget's UI
-		ChildSlot.DetachWidget();
-		ChildSlot.AttachWidget(BuildUI());
-	}
-}
-
-FString SCineAssemblyEditWidget::GetAssemblyName()
-{
-	if (CineAssembly)
-	{
-		FString AssemblyName;
-		CineAssembly->GetName(AssemblyName);
-		return AssemblyName;
-	}
-	return TEXT("CineAssembly");
-}
-
-bool SCineAssemblyEditWidget::HasRenderedThumbnail()
-{
-	const FObjectThumbnail* ObjectThumbnail = nullptr;
-	if (CineAssembly)
-	{
-		const FName FullAssetName = *CineAssembly->GetFullName();
-
-		FThumbnailMap ThumbnailMap;
-		ThumbnailTools::ConditionallyLoadThumbnailsForObjects({ FullAssetName }, ThumbnailMap);
-		ObjectThumbnail = ThumbnailMap.Find(FullAssetName);
-	}
-
-	const bool bHasRenderedThumbnail = ObjectThumbnail && !ObjectThumbnail->IsEmpty();
-	return bHasRenderedThumbnail;
-}
-
-TSharedRef<SWidget> SCineAssemblyEditWidget::MakeOverviewWidget()
-{
-	const UCineAssemblySchema* Schema = CineAssembly->GetSchema();
-
-	auto GetVisbility = [Schema]() -> EVisibility
-		{
-			return Schema ? EVisibility::Visible : EVisibility::Collapsed;
-		};
-
-	TSharedRef<FAssetThumbnail> AssetThumbnail = MakeShared<FAssetThumbnail>(CineAssembly, 256, 256, UThumbnailManager::Get().GetSharedThumbnailPool());
-	FAssetThumbnailConfig ThumbnailConfig;
-	ThumbnailConfig.AssetTypeColorOverride = FLinearColor::Transparent;
-
-	return SNew(SVerticalBox)
-
-		+ SVerticalBox::Slot()
-		.AutoHeight()
-		.Padding(0.0f, 0.0f, 0.0f, 16.0f)
-		[
-			SNew(SBorder)
-				.BorderImage(FCineAssemblyToolsStyle::Get().GetBrush("ProductionWizard.RecessedBackground"))
-				.Padding(4.0f)
-				[
-					SNew(SVerticalBox)
-
-					+SVerticalBox::Slot()
-						.AutoHeight()
-						.Padding(0.0f, 8.0f, 0.0f, 8.0f)
-						.HAlign(HAlign_Center)
-						[
-							SNew(SHorizontalBox)
-
-								+ SHorizontalBox::Slot()
-								.AutoWidth()
-								[
-									AssetThumbnail->MakeThumbnailWidget(ThumbnailConfig)
-								]
-
-								+ SHorizontalBox::Slot()
-								.Padding(16.0f, 0.0, 0.0, 0.0f)
-								.VAlign(VAlign_Center)
-								.FillContentWidth(1.0f)
-								[
-									SNew(STextBlock)
-										.Text(LOCTEXT("ThumbnailHintText", "This assembly does not currently have a preview thumbnail. Open this asset in Sequencer and save it to render a preview to display here."))
-										.AutoWrapText(true)
-										.Visibility_Lambda([this]() { return HasRenderedThumbnail() ? EVisibility::Collapsed : EVisibility::Visible; })
-								]
-						]
-
-						+ SVerticalBox::Slot()
-						.AutoHeight()
-						.HAlign(HAlign_Center)
-						.Padding(0.0f, 0.0f, 0.0f, 4.0f)
-						[
-							SNew(STextBlock).Text_Lambda([this]() {	return FText::FromString(GetAssemblyName()); })
-						]
-
-						+ SVerticalBox::Slot()
-						.AutoHeight()
-						.HAlign(HAlign_Center)
-						.Padding(0.0f, 0.0f, 0.0f, 8.0f)
-						[
-							SNew(STextBlock)
-								.Text_Lambda([Schema]() { return Schema ? FText::FromString(Schema->SchemaName) : LOCTEXT("NoSchemaName", "No Schema"); })
-								.ColorAndOpacity(FSlateColor::UseSubduedForeground())
-						]
-				]
-		]
-
-		+ SVerticalBox::Slot()
-		.MinHeight(300.0f)
-		.FillContentHeight(1.0f)
-		[
-			SNew(SBorder)
-				.BorderImage(FCineAssemblyToolsStyle::Get().GetBrush("Borders.Background"))
-				.Padding(16.0f)
-				[
-					SNew(SMultiLineEditableText)
-						.HintText(LOCTEXT("NoteHintText", "Assembly Notes"))
-						.Text_Lambda([this]()
-							{
-								return FText::FromString(CineAssembly->AssemblyNote);
-							})
-						.OnTextCommitted_Lambda([this](const FText& InText, ETextCommit::Type InCommitType)
-							{
-								if (CineAssembly->AssemblyNote != InText.ToString())
-								{
-									CineAssembly->Modify();
-									CineAssembly->AssemblyNote = InText.ToString();
-								}
-							})
-				]
-		];
-}
-
 #undef LOCTEXT_NAMESPACE
diff --git a/Engine/Plugins/VirtualProduction/CinematicAssemblyTools/Source/CineAssemblyToolsEditor/Private/UI/CineAssembly/SCineAssemblyConfigWindow.h b/Engine/Plugins/VirtualProduction/CinematicAssemblyTools/Source/CineAssemblyToolsEditor/Private/UI/CineAssembly/SCineAssemblyConfigWindow.h
--- a/Engine/Plugins/VirtualProduction/CinematicAssemblyTools/Source/CineAssemblyToolsEditor/Private/UI/CineAssembly/SCineAssemblyConfigWindow.h
+++ b/Engine/Plugins/VirtualProduction/CinematicAssemblyTools/Source/CineAssemblyToolsEditor/Private/UI/CineAssembly/SCineAssemblyConfigWindow.h
@@ -6,33 +6,8 @@
 
 #include "CineAssembly.h"
 #include "CineAssemblySchema.h"
-#include "IDetailsView.h"
-#include "Widgets/Layout/SWidgetSwitcher.h"
+#include "SCineAssemblyConfigPanel.h"
 #include "Widgets/SWindow.h"
-#include "Widgets/Views/STreeView.h"
-
-/** An entry in the hierarchy tree view */
-struct FHierarchyTreeItem
-{
-	/** The types of items that can be represented in this tree view */
-	enum class EItemType : uint8
-	{
-		Asset,
-		Folder
-	};
-
-	/** The type of this tree item */
-	EItemType Type;
-
-	/** The relative path of this tree item (possibly containing tokens) */
-	FTemplateString Path;
-
-	/** The children of this item in the tree that are Asset types */
-	TArray<TSharedPtr<FHierarchyTreeItem>> ChildAssets;
-
-	/** The children of this item in the tree that are Folder types */
-	TArray<TSharedPtr<FHierarchyTreeItem>> ChildFolders;
-};
 
 /**
  * A window for configuring the properties of a UCineAssembly asset
@@ -49,24 +24,12 @@ public:
 	void Construct(const FArguments& InArgs, const FString& InCreateAssetPath);
 
 private:
-	/** Creates the panel that displays the available Cine Assembly templates */
-	TSharedRef<SWidget> MakeCineTemplatePanel();
-
-	/** Creates the panel that holds the various tabs with properties of the Cine Assembly */
-	TSharedRef<SWidget> MakeInfoPanel();
+	/** Creates the panel that displays the available Cine Assembly Schemas */
+	TSharedRef<SWidget> MakeSchemaPanel();
 
 	/** Creates the buttons on the bottom of the window */
 	TSharedRef<SWidget> MakeButtonsPanel();
 
-	/** Creates the widget to display for the Overview tab */
-	TSharedRef<SWidget> MakeDetailsWidget();
-
-	/** Creates the widget to display for the Hierarchy tab */
-	TSharedRef<SWidget> MakeHierarchyWidget();
-
-	/** Creates the widget to display for the Notes tab */
-	TSharedRef<SWidget> MakeNotesWidget();
-
 	/** Returns the text to display on the Create Asset button, based on the selected schema */
 	FText GetCreateButtonText() const;
 
@@ -79,106 +42,17 @@ private:
 	/** Updates the UI and CineAssembly properties based on the selected schema */
 	void OnSchemaSelected(const FAssetData& AssetData);
 
-	/** 
-	 * Evaluates the input template string with the naming tokens subsystem, and stores the result in the Resolved text.
-	 * This function is throttled to only run at a set frequency, to avoid the potential to constantly query the naming tokens subsystem.
-	 */
-	void EvaluateTokenString(FTemplateString& StringToEvaluate);
-
-	/** Populate the tree view items from the list of folders and assets specified by the selected schema */
-	void PopulateHierarchyTree();
-
-	/** Generates the row widget for an entry in the tree view */
-	TSharedRef<ITableRow> OnGenerateTreeRow(TSharedPtr<FHierarchyTreeItem> TreeItem, const TSharedRef<STableViewBase>& OwnerTable);
-
-	/** Gets the children of the input tree view item to build additional tree rows */
-	void OnGetChildren(TSharedPtr<FHierarchyTreeItem> TreeItem, TArray<TSharedPtr<FHierarchyTreeItem>>& OutNodes);
-
-	/** Returns the tree item whose path matches the input path */
-	TSharedPtr<FHierarchyTreeItem> FindItemAtPathRecursive(TSharedPtr<FHierarchyTreeItem> TreeItem, const FString& Path) const;
-
-	/** Recursively expands every item in the tree view */
-	void ExpandTreeRecursive(TSharedPtr<FHierarchyTreeItem> TreeItem) const;
-
-	/** Recursively evaluates the tokens in the path of each item in the tree view */
-	void EvaluateHierarchyTokensRecursive(TSharedPtr<FHierarchyTreeItem> TreeItem);
-
 private:
-	/** Switcher that controls which tab widget is currently visible */
-	TSharedPtr<SWidgetSwitcher> TabSwitcher;
-
-	/** Details View displaying the reflected properties of the Cine Assembly being configured */ 
-	TSharedPtr<IDetailsView> DetailsView;
-
 	/** Transient object used only by this UI to configure the properties of the new asset that will get created by the Factory */
 	TStrongObjectPtr<UCineAssembly> CineAssemblyToConfigure = nullptr;
 
-	/** The currently selected schema to use as a base for configuring the new cine assembly */
-	UCineAssemblySchema* SelectedSchema = nullptr;
+	/** The widget that display all of the configuration properties for the Cine Assembly asset */
+	TSharedPtr<SCineAssemblyConfigPanel> ConfigPanel;
 
 	/** The root path where the configured assembly will be created */
 	FString CreateAssetPath;
 
-	/** The last time the naming tokens were updated */
-	FDateTime LastTokenUpdateTime;
-
-	/** Items source for the tree view */
-	TArray<TSharedPtr<FHierarchyTreeItem>> HierarchyTreeItems;
-
-	/** The root item in the tree view */
-	TSharedPtr<FHierarchyTreeItem> RootItem;
-
-	/** A read-only tree view of folders and assets that will be created based on the selected schema */
-	TSharedPtr<STreeView<TSharedPtr<FHierarchyTreeItem>>> HierarchyTreeView;
-
 	/** Cached content browser settings, used to restore defaults when closing the window */
 	bool bShowEngineContentCached = false;
 	bool bShowPluginContentCached = false;
 };
-
-/** A panel that displays properties of a Cine Assembly asset */
-class SCineAssemblyEditWidget : public SCompoundWidget
-{
-public:
-	SCineAssemblyEditWidget() = default;
-	~SCineAssemblyEditWidget();
-
-	SLATE_BEGIN_ARGS(SCineAssemblyEditWidget) {}
-	SLATE_END_ARGS()
-
-	/** Widget construction, initialized with the assembly asset being edited */
-	void Construct(const FArguments& InArgs, UCineAssembly* InAssembly);
-
-	/**
-	 * Widget construction, initialized with the GUID of the assembly to be edited
-	 * The widget will search the asset registry to find the assembly asset with the matching GUID,
-	 * and then update the widget contents accordingly. 
-	 */
-	void Construct(const FArguments& InArgs, FGuid InAssemblyGuid);
-
-	/** Returns the name of the assembly asset being edited */
-	FString GetAssemblyName();
-
-	/** Searches the asset registry for a Cine Assembly matching the input ID and updates the UI */
-	void FindAssembly(FGuid AssemblyID);
-
-	/** Returns true if the Assembly asset has a rendered thumbnail (such as from the Sequencer preview) */
-	bool HasRenderedThumbnail();
-
-private:
-	/** Constructs the main UI for the widget */ 
-	TSharedRef<SWidget> BuildUI();
-
-	/** Creates the widget to display for the Overview tab */
-	TSharedRef<SWidget> MakeOverviewWidget();
-
-	/** Filter used by the Details View to determine which custom rows to show */
-	bool IsCustomRowVisible(FName RowName, FName ParentName);
-
-private:
-	/** Switcher that controls which tab widget is currently visible */
-	TSharedPtr<SWidgetSwitcher> TabSwitcher;
-
-	/** Transient object used only by this UI to configure the properties of the new asset that will get created by the Factory */
-	UCineAssembly* CineAssembly = nullptr;
-};
\ No newline at end of file
diff --git a/Engine/Plugins/VirtualProduction/CinematicAssemblyTools/Source/CineAssemblyToolsEditor/Private/UI/CineAssembly/SCineAssemblyPropertyEditor.cpp b/Engine/Plugins/VirtualProduction/CinematicAssemblyTools/Source/CineAssemblyToolsEditor/Private/UI/CineAssembly/SCineAssemblyPropertyEditor.cpp
new file mode 100644
--- /dev/null
+++ b/Engine/Plugins/VirtualProduction/CinematicAssemblyTools/Source/CineAssemblyToolsEditor/Private/UI/CineAssembly/SCineAssemblyPropertyEditor.cpp
@@ -0,0 +1,274 @@
+// Copyright Epic Games, Inc. All Rights Reserved.
+
+#include "SCineAssemblyPropertyEditor.h"
+
+#include "AssetRegistry/AssetRegistryModule.h"
+#include "AssetThumbnail.h"
+#include "CineAssemblyToolsStyle.h"
+#include "Modules/ModuleManager.h"
+#include "ObjectTools.h"
+#include "PropertyEditorModule.h"
+#include "ThumbnailRendering/ThumbnailManager.h"
+#include "Widgets/Input/SSegmentedControl.h"
+#include "Widgets/Layout/SBorder.h"
+#include "Widgets/Layout/SScrollBox.h"
+#include "Widgets/SBoxPanel.h"
+#include "Widgets/Text/SMultiLineEditableText.h"
+#include "Widgets/Text/STextBlock.h"
+
+#define LOCTEXT_NAMESPACE "SCineAssemblyPropertyEditor"
+
+void SCineAssemblyPropertyEditor::Construct(const FArguments& InArgs, UCineAssembly* InAssembly)
+{
+	CineAssembly = InAssembly;
+	ChildSlot[BuildUI()];
+}
+
+void SCineAssemblyPropertyEditor::Construct(const FArguments& InArgs, FGuid InAssemblyGuid)
+{
+	// The UI will be temporary because no CineAssembly has been found yet
+	ChildSlot[BuildUI()];
+
+	// If the asset registry is still scanning assets, add a callback to find the assembly asset matching the input GUID and update this widget once the scan is finished.
+	// Otherwise, we can find the assembly asset and update the UI immediately. 
+	const FAssetRegistryModule& AssetRegistryModule = FModuleManager::GetModuleChecked<FAssetRegistryModule>("AssetRegistry");
+	if (AssetRegistryModule.Get().IsLoadingAssets())
+	{
+		AssetRegistryModule.Get().OnFilesLoaded().AddSP(this, &SCineAssemblyPropertyEditor::FindAssembly, InAssemblyGuid);
+	}
+	else
+	{
+		FindAssembly(InAssemblyGuid);
+	}
+}
+
+SCineAssemblyPropertyEditor::~SCineAssemblyPropertyEditor()
+{
+	if (FAssetRegistryModule* AssetRegistryModule = FModuleManager::GetModulePtr<FAssetRegistryModule>("AssetRegistry"))
+	{
+		AssetRegistryModule->Get().OnFilesLoaded().RemoveAll(this);
+	}
+}
+
+TSharedRef<SWidget> SCineAssemblyPropertyEditor::BuildUI()
+{
+	// Build a temporary UI to display while waiting for the assembly to be loaded
+	if (!CineAssembly)
+	{
+		return SNew(SBorder)
+			.BorderImage(FCineAssemblyToolsStyle::Get().GetBrush("Borders.PanelNoBorder"))
+			.Padding(8.0f)
+			[
+				SNew(STextBlock).Text(LOCTEXT("LoadingAssemblyText", "Loading Cine Assembly..."))
+			];
+	}
+
+	FPropertyEditorModule& PropertyEditor = FModuleManager::Get().LoadModuleChecked<FPropertyEditorModule>(TEXT("PropertyEditor"));
+
+	FDetailsViewArgs DetailsViewArgs;
+	DetailsViewArgs.bAllowSearch = false;
+	DetailsViewArgs.bHideSelectionTip = true;
+
+	TSharedRef<IDetailsView> DetailsView = PropertyEditor.CreateDetailView(DetailsViewArgs);
+	DetailsView->SetIsCustomRowVisibleDelegate(FIsCustomRowVisible::CreateSP(this, &SCineAssemblyPropertyEditor::IsCustomRowVisible));
+
+	DetailsView->SetObject(CineAssembly, true);
+
+	TabSwitcher = SNew(SWidgetSwitcher)
+		+ SWidgetSwitcher::Slot()
+		[MakeOverviewWidget()]
+
+		+ SWidgetSwitcher::Slot()
+		[DetailsView];
+
+	return SNew(SBorder)
+		.BorderImage(FCineAssemblyToolsStyle::Get().GetBrush("Borders.PanelNoBorder"))
+		.Padding(8.0f)
+		[
+			SNew(SScrollBox)
+				.Orientation(Orient_Vertical)
+
+			+ SScrollBox::Slot()
+				.AutoSize()
+				.HAlign(HAlign_Center)
+				.Padding(0.0f, 0.0f, 0.0f, 8.0f)
+				[
+					SNew(SSegmentedControl<int32>)
+						.Value(0)
+						.OnValueChanged_Lambda([Switcher = TabSwitcher](int32 NewValue)
+							{
+								Switcher->SetActiveWidgetIndex(NewValue);
+							})
+
+					+ SSegmentedControl<int32>::Slot(0)
+						.Text(LOCTEXT("OverviewTab", "Overview"))
+						.Icon(FSlateIcon(FCineAssemblyToolsStyle::StyleName, "Icons.Animation").GetIcon())
+
+					+ SSegmentedControl<int32>::Slot(1)
+						.Text(LOCTEXT("DetailsTab", "Details"))
+						.Icon(FSlateIcon(FAppStyle::GetAppStyleSetName(), "Icons.Details").GetIcon())
+				]
+
+			+ SScrollBox::Slot()
+				.FillSize(1.0f)
+				.HAlign(HAlign_Fill)
+				[
+					TabSwitcher.ToSharedRef()
+				]
+		];
+}
+
+bool SCineAssemblyPropertyEditor::IsCustomRowVisible(FName RowName, FName ParentName)
+{
+	if (RowName == GET_MEMBER_NAME_CHECKED(UCineAssembly, SubAssemblyNames))
+	{
+		return false;
+	}
+	return true;
+}
+
+void SCineAssemblyPropertyEditor::FindAssembly(FGuid AssemblyID)
+{
+	const FAssetRegistryModule& AssetRegistryModule = FModuleManager::GetModuleChecked<FAssetRegistryModule>("AssetRegistry");
+
+	// The only search criterion for the asset search is for an asset with an AssemblyID matching the input GUID
+	const TMultiMap<FName, FString> TagValues = { { UCineAssembly::AssemblyGuidPropertyName, AssemblyID.ToString() } };
+
+	TArray<FAssetData> AssemblyAssets;
+	AssetRegistryModule.Get().GetAssetsByTagValues(TagValues, AssemblyAssets);
+
+	// The Assembly ID is unique, so at most one asset should ever be found
+	if (AssemblyAssets.Num() > 0)
+	{
+		CineAssembly = Cast<UCineAssembly>(AssemblyAssets[0].GetAsset());
+
+		// Update the widget's UI
+		ChildSlot.DetachWidget();
+		ChildSlot.AttachWidget(BuildUI());
+	}
+}
+
+FString SCineAssemblyPropertyEditor::GetAssemblyName()
+{
+	if (CineAssembly)
+	{
+		FString AssemblyName;
+		CineAssembly->GetName(AssemblyName);
+		return AssemblyName;
+	}
+	return TEXT("CineAssembly");
+}
+
+bool SCineAssemblyPropertyEditor::HasRenderedThumbnail()
+{
+	const FObjectThumbnail* ObjectThumbnail = nullptr;
+	if (CineAssembly)
+	{
+		const FName FullAssetName = *CineAssembly->GetFullName();
+
+		FThumbnailMap ThumbnailMap;
+		ThumbnailTools::ConditionallyLoadThumbnailsForObjects({ FullAssetName }, ThumbnailMap);
+		ObjectThumbnail = ThumbnailMap.Find(FullAssetName);
+	}
+
+	const bool bHasRenderedThumbnail = ObjectThumbnail && !ObjectThumbnail->IsEmpty();
+	return bHasRenderedThumbnail;
+}
+
+TSharedRef<SWidget> SCineAssemblyPropertyEditor::MakeOverviewWidget()
+{
+	const UCineAssemblySchema* Schema = CineAssembly->GetSchema();
+
+	auto GetVisbility = [Schema]() -> EVisibility
+		{
+			return Schema ? EVisibility::Visible : EVisibility::Collapsed;
+		};
+
+	TSharedRef<FAssetThumbnail> AssetThumbnail = MakeShared<FAssetThumbnail>(CineAssembly, 256, 256, UThumbnailManager::Get().GetSharedThumbnailPool());
+	FAssetThumbnailConfig ThumbnailConfig;
+	ThumbnailConfig.AssetTypeColorOverride = FLinearColor::Transparent;
+
+	return SNew(SVerticalBox)
+
+		+ SVerticalBox::Slot()
+		.AutoHeight()
+		.Padding(0.0f, 0.0f, 0.0f, 16.0f)
+		[
+			SNew(SBorder)
+				.BorderImage(FCineAssemblyToolsStyle::Get().GetBrush("ProductionWizard.RecessedBackground"))
+				.Padding(4.0f)
+				[
+					SNew(SVerticalBox)
+
+					+ SVerticalBox::Slot()
+						.AutoHeight()
+						.Padding(0.0f, 8.0f, 0.0f, 8.0f)
+						.HAlign(HAlign_Center)
+						[
+							SNew(SHorizontalBox)
+
+							+ SHorizontalBox::Slot()
+								.AutoWidth()
+								[
+									AssetThumbnail->MakeThumbnailWidget(ThumbnailConfig)
+								]
+
+							+ SHorizontalBox::Slot()
+								.Padding(16.0f, 0.0, 0.0, 0.0f)
+								.VAlign(VAlign_Center)
+								.FillContentWidth(1.0f)
+								[
+									SNew(STextBlock)
+										.Text(LOCTEXT("ThumbnailHintText", "This assembly does not currently have a preview thumbnail. Open this asset in Sequencer and save it to render a preview to display here."))
+										.AutoWrapText(true)
+										.Visibility_Lambda([this]() { return HasRenderedThumbnail() ? EVisibility::Collapsed : EVisibility::Visible; })
+								]
+						]
+
+					+ SVerticalBox::Slot()
+						.AutoHeight()
+						.HAlign(HAlign_Center)
+						.Padding(0.0f, 0.0f, 0.0f, 4.0f)
+						[
+							SNew(STextBlock).Text_Lambda([this]() {	return FText::FromString(GetAssemblyName()); })
+						]
+
+					+ SVerticalBox::Slot()
+						.AutoHeight()
+						.HAlign(HAlign_Center)
+						.Padding(0.0f, 0.0f, 0.0f, 8.0f)
+						[
+							SNew(STextBlock)
+								.Text_Lambda([Schema]() { return Schema ? FText::FromString(Schema->SchemaName) : LOCTEXT("NoSchemaName", "No Schema"); })
+								.ColorAndOpacity(FSlateColor::UseSubduedForeground())
+						]
+				]
+		]
+
+		+ SVerticalBox::Slot()
+		.MinHeight(300.0f)
+		.FillContentHeight(1.0f)
+		[
+			SNew(SBorder)
+				.BorderImage(FCineAssemblyToolsStyle::Get().GetBrush("Borders.Background"))
+				.Padding(16.0f)
+				[
+					SNew(SMultiLineEditableText)
+						.HintText(LOCTEXT("NoteHintText", "Assembly Notes"))
+						.Text_Lambda([this]()
+							{
+								return FText::FromString(CineAssembly->AssemblyNote);
+							})
+						.OnTextCommitted_Lambda([this](const FText& InText, ETextCommit::Type InCommitType)
+							{
+								if (CineAssembly->AssemblyNote != InText.ToString())
+								{
+									CineAssembly->Modify();
+									CineAssembly->AssemblyNote = InText.ToString();
+								}
+							})
+				]
+		];
+}
+
+#undef LOCTEXT_NAMESPACE
diff --git a/Engine/Plugins/VirtualProduction/CinematicAssemblyTools/Source/CineAssemblyToolsEditor/Private/UI/CineAssembly/SCineAssemblyPropertyEditor.h b/Engine/Plugins/VirtualProduction/CinematicAssemblyTools/Source/CineAssemblyToolsEditor/Private/UI/CineAssembly/SCineAssemblyPropertyEditor.h
new file mode 100644
--- /dev/null
+++ b/Engine/Plugins/VirtualProduction/CinematicAssemblyTools/Source/CineAssemblyToolsEditor/Private/UI/CineAssembly/SCineAssemblyPropertyEditor.h
@@ -0,0 +1,57 @@
+// Copyright Epic Games, Inc. All Rights Reserved.
+
+#pragma once
+
+#include "Widgets/SCompoundWidget.h"
+
+#include "CineAssembly.h"
+#include "Widgets/Layout/SWidgetSwitcher.h"
+
+/** 
+ * A panel that displays the editable properties of a Cine Assembly asset 
+ */
+class SCineAssemblyPropertyEditor : public SCompoundWidget
+{
+public:
+	SCineAssemblyPropertyEditor() = default;
+	~SCineAssemblyPropertyEditor();
+
+	SLATE_BEGIN_ARGS(SCineAssemblyPropertyEditor) {}
+	SLATE_END_ARGS()
+
+	/** Widget construction, initialized with the assembly asset being edited */
+	void Construct(const FArguments& InArgs, UCineAssembly* InAssembly);
+
+	/**
+	 * Widget construction, initialized with the GUID of the assembly to be edited
+	 * The widget will search the asset registry to find the assembly asset with the matching GUID,
+	 * and then update the widget contents accordingly.
+	 */
+	void Construct(const FArguments& InArgs, FGuid InAssemblyGuid);
+
+	/** Returns the name of the assembly asset being edited */
+	FString GetAssemblyName();
+
+	/** Searches the asset registry for a Cine Assembly matching the input ID and updates the UI */
+	void FindAssembly(FGuid AssemblyID);
+
+	/** Returns true if the Assembly asset has a rendered thumbnail (such as from the Sequencer preview) */
+	bool HasRenderedThumbnail();
+
+private:
+	/** Constructs the main UI for the widget */
+	TSharedRef<SWidget> BuildUI();
+
+	/** Creates the widget to display for the Overview tab */
+	TSharedRef<SWidget> MakeOverviewWidget();
+
+	/** Filter used by the Details View to determine which custom rows to show */
+	bool IsCustomRowVisible(FName RowName, FName ParentName);
+
+private:
+	/** The assembly asset whose properties are displayed by this widget */
+	UCineAssembly* CineAssembly = nullptr;
+
+	/** Switcher that controls which tab widget is currently visible */
+	TSharedPtr<SWidgetSwitcher> TabSwitcher;
+};
