ShotManagement: Add shot duplication tool with configurable options for duplicating the subsequences of a cinematic assembly.

diff --git a/Engine/Plugins/VirtualProduction/CinematicAssemblyTools/Source/CineAssemblyTools/Private/CineAssembly.cpp b/Engine/Plugins/VirtualProduction/CinematicAssemblyTools/Source/CineAssemblyTools/Private/CineAssembly.cpp
--- a/Engine/Plugins/VirtualProduction/CinematicAssemblyTools/Source/CineAssemblyTools/Private/CineAssembly.cpp
+++ b/Engine/Plugins/VirtualProduction/CinematicAssemblyTools/Source/CineAssemblyTools/Private/CineAssembly.cpp
@@ -162,18 +162,15 @@ void UCineAssembly::ChangeSchema(UCineAssemblySchema* InSchema)
 #if WITH_EDITOR
 void UCineAssembly::CreateSubAssemblies()
 {
-	if (!BaseSchema || !SubAssemblies.IsEmpty())
+	if (!BaseSchema)
 	{
 		return;
 	}
 
 	// Get the path where the top-level assembly will be created so we can create other assets relative to it
-	const FAssetData AssemblyAssetData = FAssetData(this);
-	FString PackagePath = AssemblyAssetData.PackagePath.ToString();
-
-	// Remove the default assembly path from the top-level assembly's package path to get back to the "root" path for the schema's folder hierarchy
-	FString ResolvedDefaultAssemblyPath = UCineAssemblyNamingTokens::GetResolvedText(BaseSchema->DefaultAssemblyPath, this).ToString();
-	PackagePath = PackagePath.Replace(*ResolvedDefaultAssemblyPath, TEXT(""));
+	FString AssemblyPath;
+	FString AssemblyRootFolder;
+	GetAssetPathAndRootFolder(AssemblyPath, AssemblyRootFolder);
 
 	// Create the default folders for this assembly, based on the schema
 	for (FTemplateString& FolderPath : DefaultFolderNames)
@@ -186,7 +183,7 @@ void UCineAssembly::CreateSubAssemblies()
 			continue;
 		}
 
-		const FString PathToCreate = PackagePath / FolderPath.Resolved.ToString();
+		const FString PathToCreate = AssemblyRootFolder / FolderPath.Resolved.ToString();
 		const FString RelativeFilePath = FPackageName::LongPackageNameToFilename(PathToCreate);
 		const FString AbsoluteFilePath = FPaths::ConvertRelativePathToFull(RelativeFilePath);
 
@@ -202,6 +199,12 @@ void UCineAssembly::CreateSubAssemblies()
 		}
 	}
 
+	// Early-out if the subassemblies have already been created (as would be the case if this was a duplicated assembly)
+	if (!SubAssemblies.IsEmpty())
+	{
+		return;
+	}
+
 	IAssetTools& AssetTools = FModuleManager::GetModuleChecked<FAssetToolsModule>("AssetTools").Get();
 
 	// Create a new CineAssembly for each subsequence, set its playback range to match the parent sequence, and add it to the subsequence track
@@ -224,7 +227,7 @@ void UCineAssembly::CreateSubAssemblies()
 		// Before creating each subassembly, sanity check that each one will actually have a unique name (in case there are duplicates in the schema description)
 		FString UniquePackageName;
 		FString UniqueAssetName;
-		AssetTools.CreateUniqueAssetName(PackagePath / SubAssemblyName.Resolved.ToString(), TEXT(""), UniquePackageName, UniqueAssetName);
+		AssetTools.CreateUniqueAssetName(AssemblyRootFolder / SubAssemblyName.Resolved.ToString(), TEXT(""), UniquePackageName, UniqueAssetName);
 
 		const FString SubAssemblyPath = FPaths::GetPath(UniquePackageName);
 		UCineAssembly* SubAssembly = Cast<UCineAssembly>(AssetTools.CreateAsset(UniqueAssetName, SubAssemblyPath, UCineAssembly::StaticClass(), nullptr));
@@ -496,4 +499,19 @@ void UCineAssembly::AddMetadataNamingToken(const FString& InKey)
 	CineAssemblyNamingTokens->AddMetadataToken(InKey);
 }
 
+void UCineAssembly::GetAssetPathAndRootFolder(FString& OutAssetPath, FString& OutRootFolder)
+{
+	const FAssetData AssemblyAssetData = FAssetData(this);
+	OutAssetPath = AssemblyAssetData.PackagePath.ToString();
+
+	OutRootFolder = OutAssetPath;
+	if (BaseSchema)
+	{
+		const FString DefaultAssemblyPath = BaseSchema->DefaultAssemblyPath;
+		const FString ResolvedAssemblyPath = UCineAssemblyNamingTokens::GetResolvedText(DefaultAssemblyPath, this).ToString();
+
+		OutRootFolder = OutRootFolder.Replace(*ResolvedAssemblyPath, TEXT(""));
+	}
+}
+
 #undef LOCTEXT_NAMESPACE
diff --git a/Engine/Plugins/VirtualProduction/CinematicAssemblyTools/Source/CineAssemblyTools/Public/CineAssembly.h b/Engine/Plugins/VirtualProduction/CinematicAssemblyTools/Source/CineAssemblyTools/Public/CineAssembly.h
--- a/Engine/Plugins/VirtualProduction/CinematicAssemblyTools/Source/CineAssemblyTools/Public/CineAssembly.h
+++ b/Engine/Plugins/VirtualProduction/CinematicAssemblyTools/Source/CineAssemblyTools/Public/CineAssembly.h
@@ -132,6 +132,9 @@ public:
 	/** Adds a new metadata key to the list of supported naming tokens for assemblies */
 	CINEASSEMBLYTOOLS_API void AddMetadataNamingToken(const FString& InKey);
 
+	/** Returns the asset path and its root folder (which might be different if the schema defines a default assembly path) */
+	CINEASSEMBLYTOOLS_API void GetAssetPathAndRootFolder(FString& OutAssetPath, FString& OutRootFolder);
+
 private:
 	/** Sets the base schema for this assembly and re-initializes the metadata inherited from the schema */
 	CINEASSEMBLYTOOLS_API void ChangeSchema(UCineAssemblySchema* InSchema);
diff --git a/Engine/Plugins/VirtualProduction/CinematicAssemblyTools/Source/CineAssemblyToolsEditor/Private/CineAssemblyFactory.cpp b/Engine/Plugins/VirtualProduction/CinematicAssemblyTools/Source/CineAssemblyToolsEditor/Private/CineAssemblyFactory.cpp
--- a/Engine/Plugins/VirtualProduction/CinematicAssemblyTools/Source/CineAssemblyToolsEditor/Private/CineAssemblyFactory.cpp
+++ b/Engine/Plugins/VirtualProduction/CinematicAssemblyTools/Source/CineAssemblyToolsEditor/Private/CineAssemblyFactory.cpp
@@ -122,17 +122,21 @@ void UCineAssemblyFactory::CreateConfiguredAssembly(UCineAssembly* ConfiguredAss
 	ConfiguredAssembly->SetFlags(Flags);
 	ConfiguredAssembly->ClearFlags(RF_Transient);
 
-	// Do the same setup that ULevelSequence assets do when they are created
-	ConfiguredAssembly->Initialize();
-
-	if (UMovieScene* MovieScene = ConfiguredAssembly->GetMovieScene())
+	// If the configured assembly already has a MovieScene object (which could happen in case of a duplication), do not overwrite it with a new one
+	if (!ConfiguredAssembly->GetMovieScene())
 	{
-		const UMovieSceneToolsProjectSettings* MovieSceneToolsProjectSettings = GetDefault<UMovieSceneToolsProjectSettings>();
+		// Do the same setup that ULevelSequence assets do when they are created
+		ConfiguredAssembly->Initialize();
 
-		const FFrameRate TickResolution = MovieScene->GetTickResolution();
-		const FFrameNumber DefaultStartFrame = (MovieSceneToolsProjectSettings->DefaultStartTime * TickResolution).FloorToFrame();
-		const int32 DefaultDuration = (MovieSceneToolsProjectSettings->DefaultDuration * TickResolution).FloorToFrame().Value;
-		MovieScene->SetPlaybackRange(DefaultStartFrame, DefaultDuration);
+		if (UMovieScene* MovieScene = ConfiguredAssembly->GetMovieScene())
+		{
+			const UMovieSceneToolsProjectSettings* MovieSceneToolsProjectSettings = GetDefault<UMovieSceneToolsProjectSettings>();
+
+			const FFrameRate TickResolution = MovieScene->GetTickResolution();
+			const FFrameNumber DefaultStartFrame = (MovieSceneToolsProjectSettings->DefaultStartTime * TickResolution).FloorToFrame();
+			const int32 DefaultDuration = (MovieSceneToolsProjectSettings->DefaultDuration * TickResolution).FloorToFrame().Value;
+			MovieScene->SetPlaybackRange(DefaultStartFrame, DefaultDuration);
+		}
 	}
 
 	ConfiguredAssembly->CreateSubAssemblies();
diff --git a/Engine/Plugins/VirtualProduction/CinematicAssemblyTools/Source/CineAssemblyToolsEditor/Private/CineAssemblyToolsEditorModule.cpp b/Engine/Plugins/VirtualProduction/CinematicAssemblyTools/Source/CineAssemblyToolsEditor/Private/CineAssemblyToolsEditorModule.cpp
--- a/Engine/Plugins/VirtualProduction/CinematicAssemblyTools/Source/CineAssemblyToolsEditor/Private/CineAssemblyToolsEditorModule.cpp
+++ b/Engine/Plugins/VirtualProduction/CinematicAssemblyTools/Source/CineAssemblyToolsEditor/Private/CineAssemblyToolsEditorModule.cpp
@@ -14,6 +14,7 @@
 #include "CineAssemblyToolsStyle.h"
 #include "ContentBrowserMenuContexts.h"
 #include "Editor.h"
+#include "Framework/Application/SlateApplication.h"
 #include "Framework/Docking/TabManager.h"
 #include "Misc/ConfigCacheIni.h"
 #include "MoviePipelineQueue.h"
@@ -26,6 +27,7 @@
 #include "UI/SProductionWizard.h"
 #include "UI/CineAssembly/SCineAssemblyPropertyEditor.h"
 #include "UI/CineAssembly/SCineAssemblySchemaWindow.h"
+#include "UI/CineAssembly/SDuplicateAssemblyWindow.h"
 #include "WorkspaceMenuStructure.h"
 #include "WorkspaceMenuStructureModule.h"
 
@@ -286,7 +288,7 @@ void FCineAssemblyToolsEditorModule::RegisterMenus()
 {
 	UToolMenus* ToolMenus = UToolMenus::Get();
 
-	// Add a Content Browser context menu action to Cine Assembly assets that allows the asset to be opened in Sequencer without loading its associated map
+	// Add Content Browser context menu actions to Cine Assembly assets
 	const FName ContextMenuName = TEXT("ContentBrowser.AssetContextMenu.CineAssembly");
 	if (UToolMenu* Menu = ToolMenus->ExtendMenu(ContextMenuName))
 	{
@@ -306,6 +308,7 @@ void FCineAssemblyToolsEditorModule::RegisterMenus()
 						const FAssetData& CineAssemblyData = SelectedAssets[0];
 						UCineAssembly* CineAssembly = Cast<UCineAssembly>(CineAssemblyData.GetAsset());
 
+						// Add an action to open the Cine Assembly asset in Sequencer without loading its associated map
 						InSection.AddMenuEntry(
 							"OpenInCurrentMapEntry",
 							LOCTEXT("OpenInCurrentMap", "Open in Current Map"),
@@ -319,6 +322,20 @@ void FCineAssemblyToolsEditorModule::RegisterMenus()
 								})),
 							EUserInterfaceActionType::None
 						);
+
+						// Add an action to open the Duplicate Assembly window
+						InSection.AddMenuEntry(
+							"DuplicateAssemblyAction",
+							LOCTEXT("DuplicateAssembly", "Duplicate Assembly"),
+							LOCTEXT("DuplicateAssemblyTooltip", "Opens the Duplicate Assembly tool"),
+							FSlateIcon(FAppStyle::GetAppStyleSetName(), "Icons.Duplicate"),
+							FUIAction(FExecuteAction::CreateLambda([CineAssembly]()
+								{
+									TSharedRef<SDuplicateAssemblyWindow> DuplicateAssemblyWindow = SNew(SDuplicateAssemblyWindow, CineAssembly);
+									FSlateApplication::Get().AddWindow(DuplicateAssemblyWindow);
+								})),
+							EUserInterfaceActionType::None
+						);
 					}
 
 					InSection.AddMenuEntry(
diff --git a/Engine/Plugins/VirtualProduction/CinematicAssemblyTools/Source/CineAssemblyToolsEditor/Private/UI/CineAssembly/SCineAssemblyConfigPanel.cpp b/Engine/Plugins/VirtualProduction/CinematicAssemblyTools/Source/CineAssemblyToolsEditor/Private/UI/CineAssembly/SCineAssemblyConfigPanel.cpp
--- a/Engine/Plugins/VirtualProduction/CinematicAssemblyTools/Source/CineAssemblyToolsEditor/Private/UI/CineAssembly/SCineAssemblyConfigPanel.cpp
+++ b/Engine/Plugins/VirtualProduction/CinematicAssemblyTools/Source/CineAssemblyToolsEditor/Private/UI/CineAssembly/SCineAssemblyConfigPanel.cpp
@@ -26,6 +26,12 @@ void SCineAssemblyConfigPanel::Construct(const FArguments& InArgs, UCineAssembly
 		+ SWidgetSwitcher::Slot()
 		[MakeNotesWidget()];
 
+	if (InArgs._HideSubAssemblies)
+	{
+		DetailsView->SetIsCustomRowVisibleDelegate(FIsCustomRowVisible::CreateSP(this, &SCineAssemblyConfigPanel::IsCustomRowVisible));
+		DetailsView->SetObject(CineAssemblyToConfigure, true);
+	}
+
 		ChildSlot
 		[
 			SNew(SVerticalBox)
@@ -210,6 +216,15 @@ TSharedRef<SWidget> SCineAssemblyConfigPanel::MakeDetailsWidget()
 		];
 }
 
+bool SCineAssemblyConfigPanel::IsCustomRowVisible(FName RowName, FName ParentName)
+{
+	if (RowName == GET_MEMBER_NAME_CHECKED(UCineAssembly, SubAssemblyNames))
+	{
+		return false;
+	}
+	return true;
+}
+
 TSharedRef<SWidget> SCineAssemblyConfigPanel::MakeHierarchyWidget()
 {
 	HierarchyTreeView = SNew(STreeView<TSharedPtr<FHierarchyTreeItem>>)
diff --git a/Engine/Plugins/VirtualProduction/CinematicAssemblyTools/Source/CineAssemblyToolsEditor/Private/UI/CineAssembly/SCineAssemblyConfigPanel.h b/Engine/Plugins/VirtualProduction/CinematicAssemblyTools/Source/CineAssemblyToolsEditor/Private/UI/CineAssembly/SCineAssemblyConfigPanel.h
--- a/Engine/Plugins/VirtualProduction/CinematicAssemblyTools/Source/CineAssemblyToolsEditor/Private/UI/CineAssembly/SCineAssemblyConfigPanel.h
+++ b/Engine/Plugins/VirtualProduction/CinematicAssemblyTools/Source/CineAssemblyToolsEditor/Private/UI/CineAssembly/SCineAssemblyConfigPanel.h
@@ -40,7 +40,13 @@ class SCineAssemblyConfigPanel : public SCompoundWidget
 public:
 	SCineAssemblyConfigPanel() = default;
 
-	SLATE_BEGIN_ARGS(SCineAssemblyConfigPanel) {}
+	SLATE_BEGIN_ARGS(SCineAssemblyConfigPanel) 
+		: _HideSubAssemblies(false)
+		{}
+
+		/** If set to true, the SubAssemblies section in the details view will be hidden */
+		SLATE_ARGUMENT(bool, HideSubAssemblies)
+
 	SLATE_END_ARGS()
 
 	void Construct(const FArguments& InArgs, UCineAssembly* InAssembly);
@@ -82,6 +88,9 @@ private:
 	 */
 	void EvaluateTokenString(FTemplateString& StringToEvaluate);
 
+	/** Filter used by the Details View to determine which custom rows to show */
+	bool IsCustomRowVisible(FName RowName, FName ParentName);
+
 private:
 	/** Transient object used only by this UI to configure the properties of the new asset that will get created by the Factory */
 	UCineAssembly* CineAssemblyToConfigure = nullptr;
diff --git a/Engine/Plugins/VirtualProduction/CinematicAssemblyTools/Source/CineAssemblyToolsEditor/Private/UI/CineAssembly/SDuplicateAssemblyWindow.cpp b/Engine/Plugins/VirtualProduction/CinematicAssemblyTools/Source/CineAssemblyToolsEditor/Private/UI/CineAssembly/SDuplicateAssemblyWindow.cpp
new file mode 100644
--- /dev/null
+++ b/Engine/Plugins/VirtualProduction/CinematicAssemblyTools/Source/CineAssemblyToolsEditor/Private/UI/CineAssembly/SDuplicateAssemblyWindow.cpp
@@ -0,0 +1,735 @@
+// Copyright Epic Games, Inc. All Rights Reserved.
+
+#include "SDuplicateAssemblyWindow.h"
+
+#include "AssetToolsModule.h"
+#include "CineAssemblyFactory.h"
+#include "CineAssemblyNamingTokens.h"
+#include "CineAssemblyToolsStyle.h"
+#include "ContentBrowserModule.h"
+#include "IContentBrowserSingleton.h"
+#include "MovieScene.h"
+#include "MovieSceneFolder.h"
+#include "ObjectTools.h"
+#include "SCineAssemblyConfigPanel.h"
+#include "Sections/MovieSceneSubSection.h"
+#include "STemplateStringEditableTextBox.h"
+#include "Tracks/MovieSceneSubTrack.h"
+#include "Widgets/Input/SEditableTextBox.h"
+#include "Widgets/Layout/SSeparator.h"
+#include "Widgets/Layout/SSpacer.h"
+#include "Widgets/Text/STextBlock.h"
+
+#define LOCTEXT_NAMESPACE "SDuplicateAssemblyWindow"
+
+FName SDuplicateAssemblyWindow::OriginalNameColumn = "OriginalName";
+FName SDuplicateAssemblyWindow::DuplicationModeColumn = "DuplicationMode";
+FName SDuplicateAssemblyWindow::DuplicateNameColumn = "DuplicateName";
+
+void SDuplicateAssemblyWindow::Construct(const FArguments& InArgs, UCineAssembly* InAssembly)
+{
+	OriginalAssembly = InAssembly;
+
+	// Duplicate the original assembly into a new transient assembly that can be configured in the UI
+	// Note: all of the subsequence tracks/sections will be duplicated, but will reference the sequences in the original assembly
+	// So if the user chooses MaintainReference as the duplication mode for any subsequence, no further work is needed for that subsequence
+	DuplicateAssembly.Reset(Cast<UCineAssembly>(StaticDuplicateObject(OriginalAssembly, GetTransientPackage(), NAME_None, RF_Transient)));
+
+	// Get the path and root folder of the original assembly) and initialize the duplication path to be the same root folder as the original
+	OriginalAssembly->GetAssetPathAndRootFolder(OriginalAssemblyPath, OriginalAssemblyRootFolder);
+	DuplicationPath = OriginalAssemblyRootFolder;
+
+	// This path will be used to find the relative paths for some subsequences, but the function that finds the relative path only works if there is a trailing '/'
+	OriginalAssemblyPath = OriginalAssemblyPath.AppendChar('/');
+
+	const FVector2D DefaultWindowSize = FVector2D(1200, 750);
+
+	SWindow::Construct(SWindow::FArguments()
+		.Title(LOCTEXT("WindowTitle", "Duplicate Assembly"))
+		.SizingRule(ESizingRule::UserSized)
+		.ClientSize(DefaultWindowSize)
+		.SupportsMaximize(false)
+		.SupportsMinimize(false)
+		[
+			SNew(SBorder)
+				.BorderImage(FCineAssemblyToolsStyle::Get().GetBrush("ProductionWizard.PanelBackground"))
+				[
+					SNew(SVerticalBox)
+
+					+ SVerticalBox::Slot()
+						[
+							SNew(SSplitter)
+								.Orientation(Orient_Horizontal)
+								.PhysicalSplitterHandleSize(2.0f)
+
+							+ SSplitter::Slot()
+								.Value(0.65f)
+								[
+									MakeSubsequencesPanel()
+								]
+
+							+ SSplitter::Slot()
+								.Value(0.35f)
+								[
+									SNew(SCineAssemblyConfigPanel, DuplicateAssembly.Get())
+										.HideSubAssemblies(true)
+								]
+						]
+
+					+ SVerticalBox::Slot()
+						.AutoHeight()
+						[
+							SNew(SSeparator)
+								.Orientation(Orient_Horizontal)
+								.Thickness(2.0f)
+						]
+
+					+ SVerticalBox::Slot()
+						.AutoHeight()
+						[
+							MakeButtonsPanel()
+						]
+				]
+		]);
+}
+
+TSharedRef<SWidget> SDuplicateAssemblyWindow::MakeSubsequencesPanel()
+{
+	// Walk the duplicate assembly to recursively find every subsequence and build up the subsequence tree
+	PopulateSubsequenceTreeRecurisve(DuplicateAssembly.Get(), nullptr);
+
+	TreeView = SNew(STreeView<TSharedPtr<FSequenceTreeItem>>)
+		.TreeItemsSource(&TreeItems)
+		.SelectionMode(ESelectionMode::None)
+		.OnGenerateRow(this, &SDuplicateAssemblyWindow::OnGenerateTreeRow)
+		.OnGetChildren(this, &SDuplicateAssemblyWindow::OnGetChildren)
+		.HeaderRow
+		(
+			SNew(SHeaderRow)
+
+			+ SHeaderRow::Column(OriginalNameColumn)
+			.DefaultLabel(LOCTEXT("OriginalNameColumn", "Original Name"))
+			.FillWidth(0.3f)
+			.HeaderContentPadding(FMargin(8.0f, 0.0f, 0.0f, 0.0f))
+
+			+ SHeaderRow::Column(DuplicationModeColumn)
+			.DefaultLabel(LOCTEXT("DuplicationModeColumn", "Duplication Mode"))
+			.FillWidth(0.35f)
+			.HeaderContentPadding(FMargin(8.0f, 0.0f, 0.0f, 0.0f))
+
+			+ SHeaderRow::Column(DuplicateNameColumn)
+			.DefaultLabel(LOCTEXT("DuplicateNameColumn", "Duplicate Name"))
+			.FillWidth(0.35f)
+			.HeaderContentPadding(FMargin(8.0f, 0.0f, 0.0f, 0.0f))
+		);
+
+	// Start with the entire subsequence tree expanded
+	for (TSharedPtr<FSequenceTreeItem>& Item : TreeItems)
+	{
+		ExpandTreeItem(Item);
+	}
+
+	// Create a path picker widget to let the user choose where the duplicate assembly should be created
+	const FContentBrowserModule& ContentBrowserModule = FModuleManager::GetModuleChecked<FContentBrowserModule>("ContentBrowser");
+
+	FPathPickerConfig PathPickerConfig;
+	PathPickerConfig.DefaultPath = OriginalAssemblyRootFolder;
+	PathPickerConfig.OnPathSelected = FOnPathSelected::CreateLambda([this](const FString& NewPath)
+		{
+			DuplicationPath = NewPath;
+		});
+
+	TSharedRef<SWidget> PathPicker = ContentBrowserModule.Get().CreatePathPicker(PathPickerConfig);
+
+	// Register a Slate timer that runs at a set frequency to evaluate all of the tokens in the tree view.
+	// This will automatically be unregistered when this window is destroyed.
+	constexpr float TimerFrequency = 1.0f;
+	RegisterActiveTimer(TimerFrequency, FWidgetActiveTimerDelegate::CreateLambda([this](double InCurrentTime, float InDeltaTime) -> EActiveTimerReturnType
+		{
+			for (TSharedPtr<FSequenceTreeItem>& Item : TreeItems)
+			{
+				ResolveSubsequenceNameRecursive(Item);
+			}
+
+			return EActiveTimerReturnType::Continue;
+		}));
+
+	// Resolve the subsequence names once, immediately, so they appear correctly when the window first opens
+	for (TSharedPtr<FSequenceTreeItem>& Item : TreeItems)
+	{
+		ResolveSubsequenceNameRecursive(Item);
+	}
+
+	return SNew(SBorder)
+		.BorderImage(FCineAssemblyToolsStyle::Get().GetBrush("ProductionWizard.PanelBackground"))
+		.Padding(16.0f)
+		[
+			SNew(SVerticalBox)
+
+			+ SVerticalBox::Slot()
+				.AutoHeight()
+				.Padding(0.0f, 0.0f, 0.0f, 8.0f)
+				[
+					SNew(STextBlock)
+						.Text(LOCTEXT("DuplicateAssemblyTitle", "Subsequences"))
+						.Font(FCoreStyle::GetDefaultFontStyle("Normal", 14))
+				]
+
+			+ SVerticalBox::Slot()
+				.AutoHeight()
+				.Padding(0.0f, 0.0f, 0.0f, 16.0f)
+				[
+					SNew(STextBlock)
+						.Text(LOCTEXT("DuplicateAssemblyInstructions", "Choose how subsequences will be duplicated."))
+						.Font(FCoreStyle::GetDefaultFontStyle("Normal", 10))
+				]
+
+			+ SVerticalBox::Slot()
+				.AutoHeight()
+				.Padding(0.0f, 0.0f, 0.0f, 16.0f)
+				[
+					TreeView.ToSharedRef()
+				]
+
+			// The spacer fills the remaining space so the path picker will appear at the bottom of the panel
+			+ SVerticalBox::Slot()
+				.FillHeight(1.0f)
+				[
+					SNew(SSpacer)
+				]
+
+			+ SVerticalBox::Slot()
+				.AutoHeight()
+				[
+					SNew(SHorizontalBox)
+
+					+ SHorizontalBox::Slot()
+						.AutoWidth()
+						.VAlign(VAlign_Center)
+						.Padding(0.0f, 0.0f, 16.0f, 0.0f)
+						[
+							SNew(STextBlock).Text(LOCTEXT("DuplicationPath", "Duplication Path"))
+						]
+
+					+ SHorizontalBox::Slot()
+						.FillContentWidth(1.0f)
+						.VAlign(VAlign_Center)
+						.Padding(0.0f, 0.0f, 8.0f, 0.0f)
+						[
+							SNew(SEditableTextBox)
+								.Text_Lambda([this]() { return FText::FromString(DuplicationPath); })
+								.OnTextCommitted_Lambda([this](const FText& InText, ETextCommit::Type InCommitType) { DuplicationPath = InText.ToString(); })
+						]
+
+					+ SHorizontalBox::Slot()
+						.AutoWidth()
+						.VAlign(VAlign_Center)
+						[
+							SNew(SComboButton)
+								.MenuContent()
+								[
+									SNew(SVerticalBox)
+
+									+ SVerticalBox::Slot()
+										.MaxHeight(300.0f)
+										[
+											PathPicker
+										]
+								]
+								.ButtonContent()
+								[
+	 								SNew(SImage).Image(FCineAssemblyToolsStyle::Get().GetBrush("Icons.Folder"))
+								]
+						]
+				]
+		];
+}
+
+TSharedRef<SWidget> SDuplicateAssemblyWindow::MakeButtonsPanel()
+{
+	return SNew(SBorder)
+		.BorderImage(FAppStyle::Get().GetBrush("Brushes.Panel"))
+		.Padding(16.0f)
+		.HAlign(HAlign_Right)
+		[
+			SNew(SHorizontalBox)
+
+			+ SHorizontalBox::Slot()
+				.MinWidth(118.0f)
+				.MaxWidth(118.0f)
+				.VAlign(VAlign_Center)
+				.Padding(0.0f, 0.0f, 8.0f, 0.0f)
+				[
+					SNew(SButton)
+						.Text(LOCTEXT("DuplicateButton", "Duplicate"))
+						.ButtonStyle(FAppStyle::Get(), "PrimaryButton")
+						.HAlign(HAlign_Center)
+						.OnClicked(this, &SDuplicateAssemblyWindow::OnDuplicateClicked)
+				]
+
+			+ SHorizontalBox::Slot()
+				.MinWidth(118.0f)
+				.MaxWidth(118.0f)
+				.VAlign(VAlign_Center)
+				[
+					SNew(SButton)
+						.Text(LOCTEXT("CancelButton", "Cancel"))
+						.HAlign(HAlign_Center)
+						.OnClicked(this, &SDuplicateAssemblyWindow::OnCancelClicked)
+				]
+		];
+}
+
+FReply SDuplicateAssemblyWindow::OnDuplicateClicked()
+{
+	UCineAssemblyFactory::CreateConfiguredAssembly(DuplicateAssembly.Get(), DuplicationPath);
+
+	// Flatten the duplication data from the subsequence tree into a map
+	for (TSharedPtr<FSequenceTreeItem>& Item : TreeItems)
+	{
+		BuildDuplicationMapRecursive(Item);
+	}
+
+	DuplicateAssembly->GetAssetPathAndRootFolder(DuplicateAssemblyPath, DuplicateAssemblyRootFolder);
+
+	// Duplicate any subsequences that are set to DuplicateOriginal
+	DuplicateSubsequencesRecursive(DuplicateAssembly.Get());
+
+	// Remove any subsequencers that are set to Remove
+	RemoveSubsequencesRecursive(DuplicateAssembly.Get());
+
+	RequestDestroyWindow();
+	return FReply::Handled();
+}
+
+FReply SDuplicateAssemblyWindow::OnCancelClicked()
+{
+	RequestDestroyWindow();
+	return FReply::Handled();
+}
+
+TSharedRef<ITableRow> SDuplicateAssemblyWindow::OnGenerateTreeRow(TSharedPtr<FSequenceTreeItem> TreeItem, const TSharedRef<STableViewBase>& OwnerTable)
+{
+	return SNew(SSubsequenceDuplicationRow, OwnerTable, TreeItem);
+}
+
+void SDuplicateAssemblyWindow::OnGetChildren(TSharedPtr<FSequenceTreeItem> TreeItem, TArray<TSharedPtr<FSequenceTreeItem>>& OutNodes)
+{
+	OutNodes.Append(TreeItem->Children);
+}
+
+void SDuplicateAssemblyWindow::PopulateSubsequenceTreeRecurisve(UMovieSceneSequence* InSequence, TSharedPtr<FSequenceTreeItem> InItem)
+{
+	UMovieScene* MovieScene = InSequence->GetMovieScene();
+	if (!MovieScene)
+	{
+		return;
+	}
+
+	// Loop through all of the tracks and all of the sections to find every subsequence in the input sequence
+	for (UMovieSceneTrack* Track : MovieScene->GetTracks())
+	{
+		if (UMovieSceneSubTrack* SubTrack = Cast<UMovieSceneSubTrack>(Track))
+		{
+			for (UMovieSceneSection* Section : SubTrack->GetAllSections())
+			{
+				if (UMovieSceneSubSection* SubSection = Cast<UMovieSceneSubSection>(Section))
+				{
+					if (UMovieSceneSequence* Subsequence = SubSection->GetSequence())
+					{
+						TSharedPtr<FSequenceTreeItem> NewItem = MakeShared<FSequenceTreeItem>();
+						NewItem->Subsequence = Subsequence;
+
+						// Check if this subsequence is one of the subassemblies created from the schema
+						const int32 SubAssemblyIndex = DuplicateAssembly->SubAssemblies.IndexOfByKey(SubSection);
+
+						if (SubAssemblyIndex != INDEX_NONE)
+						{
+							// If this is a subassembly, use the template name (originally defined by the schema)
+							NewItem->DuplicationData.bIsSubAssembly = true;
+							NewItem->DuplicationData.SubsequenceName.Template = DuplicateAssembly->SubAssemblyNames[SubAssemblyIndex].Template;
+						}
+						else
+						{
+							// If this is not a subassembly, find the relative path to the original assembly so that a duplicate susbequence can be created with the same relative folder structure
+							NewItem->DuplicationData.bIsSubAssembly = false;
+
+							FString SubsequencePath = FPaths::GetPath(Subsequence->GetPathName());
+							FPaths::MakePathRelativeTo(SubsequencePath, *OriginalAssemblyPath);
+
+							NewItem->DuplicationData.SubsequenceName.Template = SubsequencePath / Subsequence->GetDisplayName().ToString();
+						}
+
+						if (InItem)
+						{
+							InItem->Children.Add(NewItem);
+							NewItem->Parent = InItem;
+						}
+						else
+						{
+							TreeItems.Add(NewItem);
+							NewItem->Parent = nullptr;
+						}
+
+						PopulateSubsequenceTreeRecurisve(Subsequence, NewItem);
+					}
+				}
+			}
+		}
+	}
+}
+
+void SDuplicateAssemblyWindow::DuplicateSubsequencesRecursive(UMovieSceneSequence* Sequence)
+{
+	UMovieScene* MovieScene = Sequence->GetMovieScene();
+	if (!MovieScene)
+	{
+		return;
+	}
+
+	IAssetTools& AssetTools = FModuleManager::GetModuleChecked<FAssetToolsModule>("AssetTools").Get();
+
+	// Loop through all of the tracks and all of the sections to find every subsequence in the input sequence
+	for (UMovieSceneTrack* Track : MovieScene->GetTracks())
+	{
+		if (UMovieSceneSubTrack* SubTrack = Cast<UMovieSceneSubTrack>(Track))
+		{
+			for (UMovieSceneSection* Section : SubTrack->GetAllSections())
+			{
+				if (UMovieSceneSubSection* SubSection = Cast<UMovieSceneSubSection>(Section))
+				{
+					if (UMovieSceneSequence* Subsequence = SubSection->GetSequence())
+					{
+						const FSubsequenceDuplicationData& DuplicationData = SubsequenceDuplicationData[Subsequence];
+
+						if (DuplicationData.DuplicationMode == ECineAssemblyDuplicationMode::DuplicateOriginal)
+						{
+							const FString SubsequenceNameAndPath = UCineAssemblyNamingTokens::GetResolvedText(DuplicationData.SubsequenceName.Template, DuplicateAssembly.Get()).ToString();
+
+							FString DesiredAssetName;
+							if (DuplicationData.bIsSubAssembly)
+							{
+								// Update the display name of the subsequence track to match the resolved subassembly name
+								const FString SubAssemblyFilename = FPaths::GetBaseFilename(SubsequenceNameAndPath);
+								SubTrack->SetDisplayName(FText::FromString(SubAssemblyFilename));
+
+								DesiredAssetName = DuplicateAssemblyRootFolder / SubsequenceNameAndPath;
+							}
+							else
+							{
+								DesiredAssetName = DuplicateAssemblyPath / SubsequenceNameAndPath;
+								FPaths::CollapseRelativeDirectories(DesiredAssetName);
+							}
+
+							FString UniquePackageName;
+							FString UniqueAssetName;
+							AssetTools.CreateUniqueAssetName(DesiredAssetName, TEXT(""), UniquePackageName, UniqueAssetName);
+
+							const FString SubAssemblyPath = FPaths::GetPath(UniquePackageName);
+							UMovieSceneSequence* DuplicateSubsequence = Cast<UMovieSceneSequence>(AssetTools.DuplicateAsset(UniqueAssetName, SubAssemblyPath, Subsequence));
+
+							// Update the sequence of the duplicate subsection to be the duplicate subsequence
+							SubSection->SetSequence(DuplicateSubsequence);
+
+							// Update the data in the duplication data map to correspond to the new duplicate subsequence
+							SubsequenceDuplicationData.Add(DuplicateSubsequence, DuplicationData);
+							SubsequenceDuplicationData.Remove(Subsequence);
+						}
+
+						DuplicateSubsequencesRecursive(SubSection->GetSequence());
+					}
+				}
+			}
+		}
+	}
+}
+
+void SDuplicateAssemblyWindow::RemoveSubsequencesRecursive(UMovieSceneSequence* Sequence)
+{
+	UMovieScene* MovieScene = Sequence->GetMovieScene();
+	if (!MovieScene)
+	{
+		return;
+	}
+
+	IAssetTools& AssetTools = FModuleManager::GetModuleChecked<FAssetToolsModule>("AssetTools").Get();
+
+	TArray<UMovieSceneTrack*> TracksToRemove;
+
+	// Loop through all of the tracks and all of the sections to find every subsequence in the input sequence
+	for (UMovieSceneTrack* Track : MovieScene->GetTracks())
+	{
+		if (UMovieSceneSubTrack* SubTrack = Cast<UMovieSceneSubTrack>(Track))
+		{
+			TArray<UMovieSceneSection*> SectionsToRemove;
+
+			for (UMovieSceneSection* Section : SubTrack->GetAllSections())
+			{
+				if (UMovieSceneSubSection* SubSection = Cast<UMovieSceneSubSection>(Section))
+				{
+					if (UMovieSceneSequence* Subsequence = SubSection->GetSequence())
+					{
+						const FSubsequenceDuplicationData& DuplicationData = SubsequenceDuplicationData[Subsequence];
+						if (DuplicationData.DuplicationMode == ECineAssemblyDuplicationMode::Remove)
+						{
+							SectionsToRemove.Add(SubSection);
+						}
+						else
+						{
+							RemoveSubsequencesRecursive(Subsequence);
+						}
+					}
+				}
+			}
+
+			for (UMovieSceneSection* Section : SectionsToRemove)
+			{
+				SubTrack->Modify();
+				SubTrack->RemoveSection(*Section);
+			}
+
+			// If we have now removed every section from the subsequence track, we can remove the entire track as well
+			if (SubTrack->IsEmpty())
+			{
+				TracksToRemove.Add(SubTrack);
+			}
+		}
+	}
+
+	for (UMovieSceneTrack* Track : TracksToRemove)
+	{
+		MovieScene->RemoveTrack(*Track);
+	}
+
+	TArray<UMovieSceneFolder*> RootFolders;
+	MovieScene->GetRootFolders(RootFolders);
+
+	RemoveTracksFromFoldersRecursive(RootFolders, TracksToRemove);
+}
+
+void SDuplicateAssemblyWindow::RemoveTracksFromFoldersRecursive(TArrayView<UMovieSceneFolder* const> Folders, const TArray<UMovieSceneTrack*>& TracksToRemove)
+{
+	for (UMovieSceneFolder* Folder : Folders)
+	{
+		for (UMovieSceneTrack* Track : TracksToRemove)
+		{
+			Folder->RemoveChildTrack(Track);
+		}
+
+		RemoveTracksFromFoldersRecursive(Folder->GetChildFolders(), TracksToRemove);
+	}
+}
+
+void SDuplicateAssemblyWindow::BuildDuplicationMapRecursive(const TSharedPtr<FSequenceTreeItem> InItem)
+{
+	SubsequenceDuplicationData.Add(InItem->Subsequence, InItem->DuplicationData);
+
+	for (const TSharedPtr<FSequenceTreeItem>& Child : InItem->Children)
+	{
+		BuildDuplicationMapRecursive(Child);
+	}
+}
+
+void SDuplicateAssemblyWindow::ResolveSubsequenceNameRecursive(TSharedPtr<FSequenceTreeItem> InItem)
+{
+	// Evaluate the token template string for this tree item
+	FSubsequenceDuplicationData& DuplicationData = InItem->DuplicationData;
+	DuplicationData.SubsequenceName.Resolved = UCineAssemblyNamingTokens::GetResolvedText(DuplicationData.SubsequenceName.Template, DuplicateAssembly.Get());
+
+	for (const TSharedPtr<FSequenceTreeItem>& Child : InItem->Children)
+	{
+		ResolveSubsequenceNameRecursive(Child);
+	}
+}
+
+void SDuplicateAssemblyWindow::ExpandTreeItem(TSharedPtr<FSequenceTreeItem> InItem) const
+{
+	TreeView->SetItemExpansion(InItem, true);
+
+	for (TSharedPtr<FSequenceTreeItem>& Child : InItem->Children)
+	{
+		ExpandTreeItem(Child);
+	}
+}
+
+void SSubsequenceDuplicationRow::Construct(const FArguments& Args, const TSharedRef<STableViewBase>& OwnerTableView, const TSharedPtr<FSequenceTreeItem>& InTreeItem)
+{
+	TreeItem = InTreeItem;
+
+	FSuperRowType::FArguments StyleArguments = FSuperRowType::FArguments()
+		.Padding(FMargin(4.0f, 8.0f))
+		.Style(&FAppStyle().GetWidgetStyle<FTableRowStyle>("TableView.AlternatingRow"));
+
+	SMultiColumnTableRow<TSharedPtr<FSequenceTreeItem>>::Construct(StyleArguments, OwnerTableView);
+}
+
+TSharedRef<SWidget> SSubsequenceDuplicationRow::GenerateWidgetForColumn(const FName& ColumnName)
+{
+	if (ColumnName == SDuplicateAssemblyWindow::OriginalNameColumn)
+	{
+		return MakeOriginalNameWidget();
+	}
+	else if (ColumnName == SDuplicateAssemblyWindow::DuplicationModeColumn)
+	{
+		return MakeDuplicationModeWidget();
+	}
+	else if (ColumnName == SDuplicateAssemblyWindow::DuplicateNameColumn)
+	{
+		return MakeDuplicateNameWidget();
+	}
+
+	return SNullWidget::NullWidget;
+}
+
+TSharedRef<SWidget> SSubsequenceDuplicationRow::MakeOriginalNameWidget()
+{
+	return SNew(SHorizontalBox)
+
+		// The first column gets the expander / indent for children items in the tree view
+		+ SHorizontalBox::Slot()
+		.AutoWidth()
+		[
+			SNew(SExpanderArrow, SharedThis(this)).IndentAmount(12)
+		]
+
+		+ SHorizontalBox::Slot()
+		.AutoWidth()
+		.VAlign(VAlign_Center)
+		.Padding(8.0f, 0.0f, 0.0f, 0.0f)
+		[
+			SNew(SImage).Image(FCineAssemblyToolsStyle::Get().GetBrush("Icons.Sequencer"))
+		]
+
+		+ SHorizontalBox::Slot()
+		.VAlign(VAlign_Center)
+		.Padding(8.0f, 0.0f, 0.0f, 0.0f)
+		[
+			SNew(STextBlock).Text(TreeItem->Subsequence->GetDisplayName())
+		];
+}
+
+TSharedRef<SWidget> SSubsequenceDuplicationRow::MakeDuplicationModeWidget()
+{
+	FMenuBuilder DuplicationModeMenu(true, nullptr);
+	{
+		auto AddEntryForDuplicationMode = [this, &DuplicationModeMenu](ECineAssemblyDuplicationMode DuplicationMode)
+			{
+				DuplicationModeMenu.AddMenuEntry(
+					GetDuplicationModeDisplayName(DuplicationMode),
+					FText::GetEmpty(),
+					FSlateIcon(),
+					FUIAction(FExecuteAction::CreateSP(this, &SSubsequenceDuplicationRow::SetDuplicationModeRecursive, TreeItem, DuplicationMode)),
+					NAME_None,
+					EUserInterfaceActionType::Button
+				);
+			};
+
+		AddEntryForDuplicationMode(ECineAssemblyDuplicationMode::DuplicateOriginal);
+		AddEntryForDuplicationMode(ECineAssemblyDuplicationMode::MaintainReference);
+		AddEntryForDuplicationMode(ECineAssemblyDuplicationMode::Remove);
+	}
+
+	return SNew(SHorizontalBox)
+
+		+ SHorizontalBox::Slot()
+		.FillWidth(1.0f)
+		.Padding(4.0f, 0.0f, 16.0f, 0.0f)
+		[
+			SNew(SComboButton)
+				.MenuContent()
+				[
+					DuplicationModeMenu.MakeWidget()
+				]
+				.ButtonContent()
+				[
+					SNew(STextBlock).Text_Lambda([this]() { return GetDuplicationModeDisplayName(TreeItem->DuplicationData.DuplicationMode); })
+				]
+				.IsEnabled_Lambda([this]()
+					{
+						// Disable the duplication mode widget if the parent is MaintainReference or Remove (since the value should be locked to match the parent)
+						if (TreeItem->Parent && TreeItem->Parent->DuplicationData.DuplicationMode != ECineAssemblyDuplicationMode::DuplicateOriginal)
+						{
+							return false;
+						}
+						return true;
+					})
+		];
+}
+
+TSharedRef<SWidget> SSubsequenceDuplicationRow::MakeDuplicateNameWidget()
+{
+	return SNew(SHorizontalBox)
+
+		+ SHorizontalBox::Slot()
+		.FillWidth(1.0f)
+		.Padding(10.0f, 0.0f, 16.0f, 0.0f)
+		[
+			// TODO: STemplateStringEditableTextBox does not currently support input validation because SMultiLineEditableText does not support it.
+			// Once we add support to that widget, add some input validation here to ensure that the input is a valid object name
+			SNew(STemplateStringEditableTextBox)
+				.Text_Lambda([this]()
+					{
+						if (TreeItem->DuplicationData.DuplicationMode == ECineAssemblyDuplicationMode::DuplicateOriginal)
+						{
+							const FString SubsequenceName = FPaths::GetBaseFilename(TreeItem->DuplicationData.SubsequenceName.Template);
+							return FText::FromString(SubsequenceName);
+						}
+						return TreeItem->Subsequence->GetDisplayName();
+					})
+				.ResolvedText_Lambda([this]()
+					{
+						if (TreeItem->DuplicationData.DuplicationMode == ECineAssemblyDuplicationMode::DuplicateOriginal)
+						{
+							const FString SubsequenceName = FPaths::GetBaseFilename(TreeItem->DuplicationData.SubsequenceName.Resolved.ToString());
+							return FText::FromString(SubsequenceName);
+						}
+						return TreeItem->Subsequence->GetDisplayName();
+					})
+				.OnTextCommitted(this, &SSubsequenceDuplicationRow::OnSubsequenceNameCommitted)
+				.IsReadOnly_Lambda([this]()
+					{
+						return TreeItem->DuplicationData.DuplicationMode != ECineAssemblyDuplicationMode::DuplicateOriginal;
+					})
+				.Visibility_Lambda([this]()
+					{
+						return (TreeItem->DuplicationData.DuplicationMode != ECineAssemblyDuplicationMode::Remove) ? EVisibility::Visible : EVisibility::Hidden;
+					})
+		];
+}
+
+FText SSubsequenceDuplicationRow::GetDuplicationModeDisplayName(ECineAssemblyDuplicationMode DuplicationMode) const
+{
+	switch (DuplicationMode)
+	{
+	case ECineAssemblyDuplicationMode::DuplicateOriginal:
+		return LOCTEXT("DuplicateOriginalDisplayName", "Duplicate Original");
+	case ECineAssemblyDuplicationMode::MaintainReference:
+		return LOCTEXT("MaintainReferenceDisplayName", "Maintain Reference");
+	case ECineAssemblyDuplicationMode::Remove:
+		return LOCTEXT("RemoveDisplayName", "Remove");
+	}
+
+	return FText::GetEmpty();
+}
+
+void SSubsequenceDuplicationRow::SetDuplicationModeRecursive(TSharedPtr<FSequenceTreeItem> InItem, ECineAssemblyDuplicationMode Mode)
+{
+	InItem->DuplicationData.DuplicationMode = Mode;
+
+	if (Mode != ECineAssemblyDuplicationMode::DuplicateOriginal)
+	{
+		for (TSharedPtr<FSequenceTreeItem>& Child : InItem->Children)
+		{
+			SetDuplicationModeRecursive(Child, Mode);
+		}
+	}
+}
+
+void SSubsequenceDuplicationRow::OnSubsequenceNameCommitted(const FText& InText, ETextCommit::Type InCommitType)
+{
+	const FString OldPath = TreeItem->DuplicationData.SubsequenceName.Template;
+	const FString NewPath = FPaths::GetPath(OldPath) / InText.ToString();
+	TreeItem->DuplicationData.SubsequenceName.Template = NewPath;
+}
+
+#undef LOCTEXT_NAMESPACE
diff --git a/Engine/Plugins/VirtualProduction/CinematicAssemblyTools/Source/CineAssemblyToolsEditor/Private/UI/CineAssembly/SDuplicateAssemblyWindow.h b/Engine/Plugins/VirtualProduction/CinematicAssemblyTools/Source/CineAssemblyToolsEditor/Private/UI/CineAssembly/SDuplicateAssemblyWindow.h
new file mode 100644
--- /dev/null
+++ b/Engine/Plugins/VirtualProduction/CinematicAssemblyTools/Source/CineAssemblyToolsEditor/Private/UI/CineAssembly/SDuplicateAssemblyWindow.h
@@ -0,0 +1,169 @@
+// Copyright Epic Games, Inc. All Rights Reserved.
+
+#pragma once
+
+#include "Widgets/SWindow.h"
+
+#include "CineAssembly.h"
+#include "MovieSceneSequence.h"
+#include "Widgets/Views/STreeView.h"
+
+class UMovieSceneFolder;
+
+/** List of available duplication modes for a subsequence */
+enum class ECineAssemblyDuplicationMode : uint8
+{
+	DuplicateOriginal, // Duplicates the original assembly's subsequence track/section AND duplicates the original subsequence, updating the reference
+	MaintainReference, // Duplicates the original assembly's subsequence track/section BUT maintains the reference to the original subsequence
+	Remove // Does not duplicate the original assembly's subsequence track/section
+};
+
+/** The data required to duplicate a specific subsequence */
+struct FSubsequenceDuplicationData
+{
+	/** The duplication mode to use */
+	ECineAssemblyDuplicationMode DuplicationMode = ECineAssemblyDuplicationMode::DuplicateOriginal;
+
+	/** The name to use for the duplicated subsequence (if the mode if DuplicateOriginal) */
+	FTemplateString SubsequenceName;
+
+	/** Whether or not the subsequence was defined in the schema of the original subsequence */
+	bool bIsSubAssembly = false;
+};
+
+/** An item in the subsequence tree view, associating a subsequence with its duplication data and child subsequences */
+struct FSequenceTreeItem
+{
+	/** The subsequence represented by this item in the tree */
+	UMovieSceneSequence* Subsequence;
+
+	/** The duplication associated with this subsequence */
+	FSubsequenceDuplicationData DuplicationData;
+
+	/** Child subsequence tree items */
+	TArray<TSharedPtr<FSequenceTreeItem>> Children;
+
+	/** Parent sequence tree item */
+	TSharedPtr<FSequenceTreeItem> Parent;
+};
+
+/** Widget row for an item in the subsequence tree view */
+class SSubsequenceDuplicationRow : public SMultiColumnTableRow<TSharedPtr<FSequenceTreeItem>>
+{
+public:
+
+	SLATE_BEGIN_ARGS(SSubsequenceDuplicationRow) {}
+	SLATE_END_ARGS()
+
+	void Construct(const FArguments& Args, const TSharedRef<STableViewBase>& OwnerTableView, const TSharedPtr<FSequenceTreeItem>& InTreeItem);
+
+	/** Creates the widget for this row for the specified column */
+	virtual TSharedRef<SWidget> GenerateWidgetForColumn(const FName& ColumnName) override;
+
+private:
+	/** Make the widget for each column */
+	TSharedRef<SWidget> MakeOriginalNameWidget();
+	TSharedRef<SWidget> MakeDuplicationModeWidget();
+	TSharedRef<SWidget> MakeDuplicateNameWidget();
+
+	/** Updates the duplication mode for the input tree item, and potentially its children */
+	void SetDuplicationModeRecursive(TSharedPtr<FSequenceTreeItem> InItem, ECineAssemblyDuplicationMode Mode);
+
+	/** Updates the subsequence name associated with this tree row */
+	void OnSubsequenceNameCommitted(const FText& InText, ETextCommit::Type InCommitType);
+
+	/** Get the display name of the input duplication mode */
+	FText GetDuplicationModeDisplayName(ECineAssemblyDuplicationMode DuplicationMode) const;
+
+private:
+	/** The tree view item displayed by this row widget */
+	TSharedPtr<FSequenceTreeItem> TreeItem;
+};
+
+/** A window to configure how a Cine Assembly asset should be duplicated */
+class SDuplicateAssemblyWindow : public SWindow
+{
+public:
+	SDuplicateAssemblyWindow() = default;
+
+	SLATE_BEGIN_ARGS(SDuplicateAssemblyWindow) {}
+	SLATE_END_ARGS()
+
+	void Construct(const FArguments& InArgs, UCineAssembly* InAssembly);
+
+public:
+	static FName OriginalNameColumn;
+	static FName DuplicationModeColumn;
+	static FName DuplicateNameColumn;
+
+private:
+	/** Creates the subsequences panel */
+	TSharedRef<SWidget> MakeSubsequencesPanel();
+
+	/** Creates the buttons on the bottom of the window */
+	TSharedRef<SWidget> MakeButtonsPanel();
+
+	/** Called when the Duplicate button is pressed, which will execute the duplication steps as defined by the duplication data for each subsequence */
+	FReply OnDuplicateClicked();
+
+	/** Called when the Cancel button is pressed, which will close the window without duplicating anything */
+	FReply OnCancelClicked();
+
+	/** Generates the row widget in the subsequence tree view for the input tree item */
+	TSharedRef<ITableRow> OnGenerateTreeRow(TSharedPtr<FSequenceTreeItem> InTreeItem, const TSharedRef<STableViewBase>& OwnerTable);
+
+	/** Adds the children of the input tree item to the OutNodes array */
+	void OnGetChildren(TSharedPtr<FSequenceTreeItem> InTreeItem, TArray<TSharedPtr<FSequenceTreeItem>>& OutNodes);
+
+	/** Builds the list of tree items in the subsequence tree by recursively walking the input sequence and finding all subsequences that may need to be duplicated */
+	void PopulateSubsequenceTreeRecurisve(UMovieSceneSequence* InSequence, TSharedPtr<FSequenceTreeItem> InTreeItem);
+
+	/** Recursively duplicates any of the subsequences of the input sequence that are set to the DuplicateOriginal duplication mode */
+	void DuplicateSubsequencesRecursive(UMovieSceneSequence* InSequence);
+
+	/** Recursively removes any of the subsequences of the input sequence that are set to the Remove duplication mode */
+	void RemoveSubsequencesRecursive(UMovieSceneSequence* InSequence);
+
+	/** Recursively removes the set of tracks from the input folders (or child folders) */
+	void RemoveTracksFromFoldersRecursive(TArrayView<UMovieSceneFolder* const> Folders, const TArray<UMovieSceneTrack*>& TracksToRemove);
+
+	/** Converts the subsequence tree into a map of subsequences to their associated duplication data */
+	void BuildDuplicationMapRecursive(const TSharedPtr<FSequenceTreeItem> InTreeItem);
+
+	/** Resolves any tokens present in the subsequence name of the input tree item, then resolves its children */
+	void ResolveSubsequenceNameRecursive(TSharedPtr<FSequenceTreeItem> InTreeItem);
+
+	/** Expands the input item in the subsequence tree view */
+	void ExpandTreeItem(TSharedPtr<FSequenceTreeItem> InTreeItem) const;
+
+private:
+	/** The original assembly being duplicated */
+	UCineAssembly* OriginalAssembly;
+
+	/** The duplicated assembly that is being configured */
+	TStrongObjectPtr<UCineAssembly> DuplicateAssembly;
+
+	/** List of subsequence tree items for the tree view */
+	TArray<TSharedPtr<FSequenceTreeItem>> TreeItems;
+
+	/** Tree view of subsequences, allowing the user to configure the duplication data */
+	TSharedPtr<STreeView<TSharedPtr<FSequenceTreeItem>>> TreeView;
+
+	/** Map of subsequences to the associated duplication data, which mirrors the tree view */
+	TMap<UMovieSceneSequence*, FSubsequenceDuplicationData> SubsequenceDuplicationData;
+
+	/** The path where the duplicate assembly asset should be created */
+	FString DuplicationPath;
+
+	/** The path of the original assembly asset */
+	FString OriginalAssemblyPath;
+
+	/** The root folder of the original assembly (i.e. the asset path with the default assembly path from the schema removed) */
+	FString OriginalAssemblyRootFolder;
+
+	/** The path of the duplicate assembly asset (after it has been renamed by the factory) */
+	FString DuplicateAssemblyPath;
+
+	/** The root folder of the duplicate assembly (i.e. the asset path with the default assembly path from the schema removed) */
+	FString DuplicateAssemblyRootFolder;
+};
