ShotManagement: Fix button text being cutoff
#rnx

diff --git a/Engine/Plugins/VirtualProduction/CinematicAssemblyTools/Source/CineAssemblyToolsEditor/Private/UI/CineAssembly/SCineAssemblySchemaWindow.cpp b/Engine/Plugins/VirtualProduction/CinematicAssemblyTools/Source/CineAssemblyToolsEditor/Private/UI/CineAssembly/SCineAssemblySchemaWindow.cpp
--- a/Engine/Plugins/VirtualProduction/CinematicAssemblyTools/Source/CineAssemblyToolsEditor/Private/UI/CineAssembly/SCineAssemblySchemaWindow.cpp
+++ b/Engine/Plugins/VirtualProduction/CinematicAssemblyTools/Source/CineAssemblyToolsEditor/Private/UI/CineAssembly/SCineAssemblySchemaWindow.cpp
@@ -350,8 +350,7 @@ TSharedRef<SWidget> SCineAssemblySchemaWindow::MakeButtonsPanel()
 				SNew(SHorizontalBox)
 
 				+ SHorizontalBox::Slot()
-					.MinWidth(118.0f)
-					.MaxWidth(118.0f)
+					.AutoWidth()
 					.VAlign(VAlign_Center)
 					.Padding(0.0f, 0.0f, 8.0f, 0.0f)
 					[
