ShotManagement: Level Sequence Favorite Rating widget added to CineAssembly property edit window.

#rnx

diff --git a/Engine/Plugins/VirtualProduction/CinematicAssemblyTools/Source/CineAssemblyToolsEditor/Private/UI/CineAssembly/SCineAssemblyPropertyEditor.cpp b/Engine/Plugins/VirtualProduction/CinematicAssemblyTools/Source/CineAssemblyToolsEditor/Private/UI/CineAssembly/SCineAssemblyPropertyEditor.cpp
--- a/Engine/Plugins/VirtualProduction/CinematicAssemblyTools/Source/CineAssemblyToolsEditor/Private/UI/CineAssembly/SCineAssemblyPropertyEditor.cpp
+++ b/Engine/Plugins/VirtualProduction/CinematicAssemblyTools/Source/CineAssemblyToolsEditor/Private/UI/CineAssembly/SCineAssemblyPropertyEditor.cpp
@@ -13,6 +13,7 @@
 #include "Widgets/Layout/SBorder.h"
 #include "Widgets/Layout/SScrollBox.h"
 #include "Widgets/SBoxPanel.h"
+#include "Widgets/SLevelSequenceFavoriteRating.h"
 #include "Widgets/Text/SMultiLineEditableText.h"
 #include "Widgets/Text/STextBlock.h"
 
@@ -242,6 +243,15 @@ TSharedRef<SWidget> SCineAssemblyPropertyEditor::MakeOverviewWidget()
 								.Text_Lambda([Schema]() { return Schema ? FText::FromString(Schema->SchemaName) : LOCTEXT("NoSchemaName", "No Schema"); })
 								.ColorAndOpacity(FSlateColor::UseSubduedForeground())
 						]
+
+					+ SVerticalBox::Slot()
+						.AutoHeight()
+						.HAlign(HAlign_Center)
+						.Padding(0.0f, 0.0f, 0.0f, 8.0f)
+						[
+							SNew(SLevelSequenceFavoriteRating)
+								.LevelSequence(CineAssembly)
+						]
 				]
 		]
 
