ShotManagement: Prevent the Cine Assembly factory from renaming an assembly to "None" if the name would otherwise be invalid.
#rnx

diff --git a/Engine/Plugins/VirtualProduction/CinematicAssemblyTools/Source/CineAssemblyToolsEditor/Private/CineAssemblyFactory.cpp b/Engine/Plugins/VirtualProduction/CinematicAssemblyTools/Source/CineAssemblyToolsEditor/Private/CineAssemblyFactory.cpp
--- a/Engine/Plugins/VirtualProduction/CinematicAssemblyTools/Source/CineAssemblyToolsEditor/Private/CineAssemblyFactory.cpp
+++ b/Engine/Plugins/VirtualProduction/CinematicAssemblyTools/Source/CineAssemblyToolsEditor/Private/CineAssemblyFactory.cpp
@@ -102,8 +102,8 @@ void UCineAssemblyFactory::CreateConfiguredAssembly(UCineAssembly* ConfiguredAss
 	// Evaluate the name of the assembly
 	ConfiguredAssembly->AssemblyName.Resolved = UCineAssemblyNamingTokens::GetResolvedText(ConfiguredAssembly->AssemblyName.Template, ConfiguredAssembly);
 
-	// If the assembly name is empty, assign it a valid default name
-	if (ConfiguredAssembly->AssemblyName.Resolved.IsEmpty())
+	// If the assembly name is empty or invalid, assign it a valid default name
+	if (ConfiguredAssembly->AssemblyName.Resolved.IsEmpty() || FName(*ConfiguredAssembly->AssemblyName.Resolved.ToString()).IsNone())
 	{
 		ConfiguredAssembly->AssemblyName.Resolved = LOCTEXT("NewCineAssemblyName", "NewCineAssembly");
 		ConfiguredAssembly->AssemblyName.Template = ConfiguredAssembly->AssemblyName.Resolved.ToString();
