ShotManagement: Add the top-level {assembly} item to the hierarchy preview in the Cine Assembly config panel.
Also, fixed an issue where some folder/asset names in the hierarchy preview could have an empty name (e.g. if the tokens resolved to any empty string) by displaying the token template string instead.
#rnx

diff --git a/Engine/Plugins/VirtualProduction/CinematicAssemblyTools/Source/CineAssemblyToolsEditor/Private/UI/CineAssembly/SCineAssemblyConfigPanel.cpp b/Engine/Plugins/VirtualProduction/CinematicAssemblyTools/Source/CineAssemblyToolsEditor/Private/UI/CineAssembly/SCineAssemblyConfigPanel.cpp
--- a/Engine/Plugins/VirtualProduction/CinematicAssemblyTools/Source/CineAssemblyToolsEditor/Private/UI/CineAssembly/SCineAssemblyConfigPanel.cpp
+++ b/Engine/Plugins/VirtualProduction/CinematicAssemblyTools/Source/CineAssemblyToolsEditor/Private/UI/CineAssembly/SCineAssemblyConfigPanel.cpp
@@ -326,6 +326,17 @@ void SCineAssemblyConfigPanel::PopulateHierarchyTree()
 		AddItemsToTree(FolderTemplates, FHierarchyTreeItem::EItemType::Folder);
 		AddItemsToTree(CineAssemblyToConfigure->SubAssemblyNames, FHierarchyTreeItem::EItemType::Asset);
 
+		// Add the top-level assembly tree item
+		TSharedPtr<FHierarchyTreeItem> ParentItem = FindItemAtPathRecursive(RootItem, Schema->DefaultAssemblyPath);
+		if (ParentItem)
+		{
+			TSharedPtr<FHierarchyTreeItem> NewItem = MakeShared<FHierarchyTreeItem>();
+			NewItem->Type = FHierarchyTreeItem::EItemType::Asset;
+			NewItem->Path.Template = TEXT("{assembly}");
+
+			ParentItem->ChildAssets.Add(NewItem);
+		}
+
 		EvaluateHierarchyTokensRecursive(RootItem);
 	}
 
@@ -337,6 +348,10 @@ void SCineAssemblyConfigPanel::EvaluateHierarchyTokensRecursive(TSharedPtr<FHier
 {
 	// Evaluate the token template string for this tree item
 	TreeItem->Path.Resolved = UCineAssemblyNamingTokens::GetResolvedText(TreeItem->Path.Template, CineAssemblyToConfigure);
+	if (TreeItem->Path.Resolved.IsEmpty())
+	{
+		TreeItem->Path.Resolved = FText::FromString(TreeItem->Path.Template);
+	}
 
 	// Evaluate the tokens for all of the child assets, then resort them alphabetically based on the resolved paths
 	for (const TSharedPtr<FHierarchyTreeItem>& Asset : TreeItem->ChildAssets)
