ShotManagement: Add a feature to remember the most recent duplication preferences for subassemblies
#rnx

diff --git a/Engine/Plugins/VirtualProduction/CinematicAssemblyTools/Source/CineAssemblyToolsEditor/Private/UI/CineAssembly/SDuplicateAssemblyWindow.cpp b/Engine/Plugins/VirtualProduction/CinematicAssemblyTools/Source/CineAssemblyToolsEditor/Private/UI/CineAssembly/SDuplicateAssemblyWindow.cpp
--- a/Engine/Plugins/VirtualProduction/CinematicAssemblyTools/Source/CineAssemblyToolsEditor/Private/UI/CineAssembly/SDuplicateAssemblyWindow.cpp
+++ b/Engine/Plugins/VirtualProduction/CinematicAssemblyTools/Source/CineAssemblyToolsEditor/Private/UI/CineAssembly/SDuplicateAssemblyWindow.cpp
@@ -8,6 +8,7 @@
 #include "CineAssemblyToolsStyle.h"
 #include "ContentBrowserModule.h"
 #include "IContentBrowserSingleton.h"
+#include "Misc/ConfigCacheIni.h"
 #include "MovieScene.h"
 #include "MovieSceneFolder.h"
 #include "ObjectTools.h"
@@ -22,13 +23,16 @@
 
 #define LOCTEXT_NAMESPACE "SDuplicateAssemblyWindow"
 
-FName SDuplicateAssemblyWindow::OriginalNameColumn = "OriginalName";
-FName SDuplicateAssemblyWindow::DuplicationModeColumn = "DuplicationMode";
-FName SDuplicateAssemblyWindow::DuplicateNameColumn = "DuplicateName";
+const FName SDuplicateAssemblyWindow::OriginalNameColumn = "OriginalName";
+const FName SDuplicateAssemblyWindow::DuplicationModeColumn = "DuplicationMode";
+const FName SDuplicateAssemblyWindow::DuplicateNameColumn = "DuplicateName";
+
+const FString SDuplicateAssemblyWindow::DuplicationPreferenceSection = TEXT("CinematicAssemblyDuplication_");
 
 void SDuplicateAssemblyWindow::Construct(const FArguments& InArgs, UCineAssembly* InAssembly)
 {
 	OriginalAssembly = InAssembly;
+	Schema = OriginalAssembly->GetSchema();
 
 	// Duplicate the original assembly into a new transient assembly that can be configured in the UI
 	// Note: all of the subsequence tracks/sections will be duplicated, but will reference the sequences in the original assembly
@@ -298,6 +302,9 @@ FReply SDuplicateAssemblyWindow::OnDuplicateClicked()
 	// Remove any subsequencers that are set to Remove
 	RemoveSubsequencesRecursive(DuplicateAssembly.Get());
 
+	// Write out the duplication mode values for each subassembly to a config file so those values can be used as the default the next time an assembly with this schema type is duplicated.
+	SaveDuplicationPreferences();
+
 	RequestDestroyWindow();
 	return FReply::Handled();
 }
@@ -348,6 +355,7 @@ void SDuplicateAssemblyWindow::PopulateSubsequenceTreeRecurisve(UMovieSceneSeque
 							// If this is a subassembly, use the template name (originally defined by the schema)
 							NewItem->DuplicationData.bIsSubAssembly = true;
 							NewItem->DuplicationData.SubsequenceName.Template = DuplicateAssembly->SubAssemblyNames[SubAssemblyIndex].Template;
+							NewItem->DuplicationData.DuplicationMode = GetDuplicationPreference(NewItem->DuplicationData.SubsequenceName.Template);
 						}
 						else
 						{
@@ -358,6 +366,7 @@ void SDuplicateAssemblyWindow::PopulateSubsequenceTreeRecurisve(UMovieSceneSeque
 							FPaths::MakePathRelativeTo(SubsequencePath, *OriginalAssemblyPath);
 
 							NewItem->DuplicationData.SubsequenceName.Template = SubsequencePath / Subsequence->GetDisplayName().ToString();
+							NewItem->DuplicationData.DuplicationMode = ECineAssemblyDuplicationMode::DuplicateOriginal;
 						}
 
 						if (InItem)
@@ -552,6 +561,48 @@ void SDuplicateAssemblyWindow::ExpandTreeItem(TSharedPtr<FSequenceTreeItem> InIt
 	}
 }
 
+void SDuplicateAssemblyWindow::SaveDuplicationPreferences() const
+{
+	if (Schema && GConfig)
+	{
+		const FString SectionName = DuplicationPreferenceSection + Schema->SchemaName;
+
+		for (const TPair<UMovieSceneSequence*, FSubsequenceDuplicationData>& DuplicationData : SubsequenceDuplicationData)
+		{
+			if (DuplicationData.Value.bIsSubAssembly)
+			{
+				// The config strips off these token characters (for some reason), so we save and restore the values without them
+				FString TemplateName = DuplicationData.Value.SubsequenceName.Template;
+				TemplateName = TemplateName.Replace(TEXT("{"), TEXT(""));
+				TemplateName = TemplateName.Replace(TEXT("}"), TEXT(""));
+
+				const int32 Value = static_cast<int32>(DuplicationData.Value.DuplicationMode);
+				GConfig->SetInt(*SectionName, *TemplateName, Value, GEditorPerProjectIni);
+			}
+		}
+	}
+}
+
+ECineAssemblyDuplicationMode SDuplicateAssemblyWindow::GetDuplicationPreference(const FString& SubAssemblyName) const
+{
+	if (Schema && GConfig)
+	{
+		const FString SectionName = DuplicationPreferenceSection + Schema->SchemaName;
+
+		// The config strips off these token characters (for some reason), so we save and restore the values without them
+		FString TemplateName = SubAssemblyName;
+		TemplateName = TemplateName.Replace(TEXT("{"), TEXT(""));
+		TemplateName = TemplateName.Replace(TEXT("}"), TEXT(""));
+
+		int32 Value = 0;
+		GConfig->GetInt(*SectionName, *TemplateName, Value, GEditorPerProjectIni);
+
+		return static_cast<ECineAssemblyDuplicationMode>(Value);
+	}
+
+	return ECineAssemblyDuplicationMode::DuplicateOriginal;
+}
+
 void SSubsequenceDuplicationRow::Construct(const FArguments& Args, const TSharedRef<STableViewBase>& OwnerTableView, const TSharedPtr<FSequenceTreeItem>& InTreeItem)
 {
 	TreeItem = InTreeItem;
diff --git a/Engine/Plugins/VirtualProduction/CinematicAssemblyTools/Source/CineAssemblyToolsEditor/Private/UI/CineAssembly/SDuplicateAssemblyWindow.h b/Engine/Plugins/VirtualProduction/CinematicAssemblyTools/Source/CineAssemblyToolsEditor/Private/UI/CineAssembly/SDuplicateAssemblyWindow.h
--- a/Engine/Plugins/VirtualProduction/CinematicAssemblyTools/Source/CineAssemblyToolsEditor/Private/UI/CineAssembly/SDuplicateAssemblyWindow.h
+++ b/Engine/Plugins/VirtualProduction/CinematicAssemblyTools/Source/CineAssemblyToolsEditor/Private/UI/CineAssembly/SDuplicateAssemblyWindow.h
@@ -92,9 +92,9 @@ public:
 	void Construct(const FArguments& InArgs, UCineAssembly* InAssembly);
 
 public:
-	static FName OriginalNameColumn;
-	static FName DuplicationModeColumn;
-	static FName DuplicateNameColumn;
+	static const FName OriginalNameColumn;
+	static const FName DuplicationModeColumn;
+	static const FName DuplicateNameColumn;
 
 private:
 	/** Creates the subsequences panel */
@@ -136,10 +136,19 @@ private:
 	/** Expands the input item in the subsequence tree view */
 	void ExpandTreeItem(TSharedPtr<FSequenceTreeItem> InTreeItem) const;
 
+	/** Save the duplication mode settings for each subassembly to a config file to re-use the next time an assembly with the same schema is duplicated */
+	void SaveDuplicationPreferences() const;
+
+	/** Reads the saved config and returns the most recent duplication mode for the input subassembly */
+	ECineAssemblyDuplicationMode GetDuplicationPreference(const FString& SubAssemblyName) const;
+
 private:
 	/** The original assembly being duplicated */
 	UCineAssembly* OriginalAssembly;
 
+	/** The schema of the assembly being duplicated */
+	const UCineAssemblySchema* Schema;
+
 	/** The duplicated assembly that is being configured */
 	TStrongObjectPtr<UCineAssembly> DuplicateAssembly;
 
@@ -166,4 +175,7 @@ private:
 
 	/** The root folder of the duplicate assembly (i.e. the asset path with the default assembly path from the schema removed) */
 	FString DuplicateAssemblyRootFolder;
+
+	/** Config section name for duplication preferences */
+	static const FString DuplicationPreferenceSection;
 };
