ShotManagement: Limit Cine Assembly and Cine Assembly Schema naming lengths to prevent issues with exceed name maximum lengths.

Create buttons are disabled if the name is too long.
Tooltip is shown on create buttons when the name is too long, stating that the name must be less than the maximum.
Check added to creation functions to prevent any names longer than the maximum.

diff --git a/Engine/Plugins/VirtualProduction/CinematicAssemblyTools/Source/CineAssemblyToolsEditor/Private/CineAssemblyFactory.cpp b/Engine/Plugins/VirtualProduction/CinematicAssemblyTools/Source/CineAssemblyToolsEditor/Private/CineAssemblyFactory.cpp
--- a/Engine/Plugins/VirtualProduction/CinematicAssemblyTools/Source/CineAssemblyToolsEditor/Private/CineAssemblyFactory.cpp
+++ b/Engine/Plugins/VirtualProduction/CinematicAssemblyTools/Source/CineAssemblyToolsEditor/Private/CineAssemblyFactory.cpp
@@ -2,6 +2,7 @@
 
 #include "CineAssemblyFactory.h"
 
+#include "AssetViewUtils.h"
 #include "AssetToolsModule.h"
 #include "CineAssemblyNamingTokens.h"
 #include "CineAssemblyToolsAnalytics.h"
@@ -57,7 +58,7 @@ UObject* UCineAssemblyFactory::FactoryCreateNew(UClass* Class, UObject* InParent
 	return NewAssembly;
 }
 
-void UCineAssemblyFactory::MakeUniqueNameAndPath(UCineAssembly* ConfiguredAssembly, const FString& CreateAssetPath, FString& UniquePackageName, FString& UniqueAssetName)
+FString UCineAssemblyFactory::MakeAssemblyPackageName(UCineAssembly* ConfiguredAssembly, const FString& CreateAssetPath)
 {
 	// Resolve the default assembly path
 	FString DefaultAssemblyPath;
@@ -69,9 +70,21 @@ void UCineAssemblyFactory::MakeUniqueNameAndPath(UCineAssembly* ConfiguredAssemb
 		}
 	}
 
+	return CreateAssetPath / DefaultAssemblyPath / ConfiguredAssembly->AssemblyName.Resolved.ToString();
+}
+
+bool UCineAssemblyFactory::MakeUniqueNameAndPath(UCineAssembly* ConfiguredAssembly, const FString& CreateAssetPath, FString& UniquePackageName, FString& UniqueAssetName)
+{
+	const FString DesiredPackageName = MakeAssemblyPackageName(ConfiguredAssembly, CreateAssetPath);
+
+	// Ensure the package name length does not exceed the maximum cook path length as this may cause issues later on.
+	if (FText OutErrorMessage; !AssetViewUtils::IsValidPackageForCooking(DesiredPackageName, OutErrorMessage))
+	{
+		return false;
+	}
+
 	// Ensure that the resolved assembly name is actually unique
 	const FString AssemblyName = ConfiguredAssembly->AssemblyName.Resolved.ToString();
-	const FString DesiredPackageName = CreateAssetPath / DefaultAssemblyPath / AssemblyName;
 	const FString DesiredSuffix = TEXT("");
 
 	IAssetTools& AssetTools = FModuleManager::GetModuleChecked<FAssetToolsModule>("AssetTools").Get();
@@ -87,8 +100,10 @@ void UCineAssemblyFactory::MakeUniqueNameAndPath(UCineAssembly* ConfiguredAssemb
 		// By updating the assembly name, we may also need to reevaluate the assembly path.
 		// However, once we change the assembly path, we need to verify uniqueness again (in the new path).
 		// Therefore, we recursively call this function until we end up with a combination of a unique assembly path and name.
-		MakeUniqueNameAndPath(ConfiguredAssembly, CreateAssetPath, UniquePackageName, UniqueAssetName);
+		return MakeUniqueNameAndPath(ConfiguredAssembly, CreateAssetPath, UniquePackageName, UniqueAssetName);
 	}
+
+	return true;
 }
 
 void UCineAssemblyFactory::CreateConfiguredAssembly(UCineAssembly* ConfiguredAssembly, const FString& CreateAssetPath)
@@ -111,7 +126,10 @@ void UCineAssemblyFactory::CreateConfiguredAssembly(UCineAssembly* ConfiguredAss
 
 	FString UniquePackageName;
 	FString UniqueAssetName;
-	MakeUniqueNameAndPath(ConfiguredAssembly, CreateAssetPath, UniquePackageName, UniqueAssetName);
+	if (!MakeUniqueNameAndPath(ConfiguredAssembly, CreateAssetPath, UniquePackageName, UniqueAssetName))
+	{
+		return;
+	}
 
 	// The input assembly object was created in the transient package while its properties were configured.
 	// Now, we can create a real package for it, rename it, and update its object flags.
diff --git a/Engine/Plugins/VirtualProduction/CinematicAssemblyTools/Source/CineAssemblyToolsEditor/Private/CineAssemblyFactory.h b/Engine/Plugins/VirtualProduction/CinematicAssemblyTools/Source/CineAssemblyToolsEditor/Private/CineAssemblyFactory.h
--- a/Engine/Plugins/VirtualProduction/CinematicAssemblyTools/Source/CineAssemblyToolsEditor/Private/CineAssemblyFactory.h
+++ b/Engine/Plugins/VirtualProduction/CinematicAssemblyTools/Source/CineAssemblyToolsEditor/Private/CineAssemblyFactory.h
@@ -20,8 +20,11 @@ public:
 	/** Takes a pre-configured, transient assembly, creates a valid package for it, and initializes it */
 	static void CreateConfiguredAssembly(UCineAssembly* ConfiguredAssembly, const FString& CreateAssetPath);
 
+	/** Takes a CineAssembly and a creation path and makes a formatted CineAssembly package name. */
+	static FString MakeAssemblyPackageName(UCineAssembly* ConfiguredAssembly, const FString& CreateAssetPath);
+
 	/** Recursively evaluates the default assembly name and path until a unique combination is found */
-	static void MakeUniqueNameAndPath(UCineAssembly* ConfiguredAssembly, const FString& CreateAssetPath, FString& UniquePackageName, FString& UniqueAssetName);
+	static bool MakeUniqueNameAndPath(UCineAssembly* ConfiguredAssembly, const FString& CreateAssetPath, FString& UniquePackageName, FString& UniqueAssetName);
 
 protected:
 	UCineAssemblyFactory();
diff --git a/Engine/Plugins/VirtualProduction/CinematicAssemblyTools/Source/CineAssemblyToolsEditor/Private/CineAssemblySchemaCustomization.cpp b/Engine/Plugins/VirtualProduction/CinematicAssemblyTools/Source/CineAssemblyToolsEditor/Private/CineAssemblySchemaCustomization.cpp
--- a/Engine/Plugins/VirtualProduction/CinematicAssemblyTools/Source/CineAssemblyToolsEditor/Private/CineAssemblySchemaCustomization.cpp
+++ b/Engine/Plugins/VirtualProduction/CinematicAssemblyTools/Source/CineAssemblyToolsEditor/Private/CineAssemblySchemaCustomization.cpp
@@ -4,6 +4,7 @@
 
 #include "Algo/Contains.h"
 #include "AssetRegistry/AssetRegistryModule.h"
+#include "AssetViewUtils.h"
 #include "CineAssemblySchema.h"
 #include "DetailLayoutBuilder.h"
 #include "PropertyCustomizationHelpers.h"
@@ -136,6 +137,19 @@ bool FCineAssemblySchemaCustomization::ValidateSchemaName(const FText& InText, F
 		return false;
 	}
 
+	/* Check the resulting package name would not fail validation if we have enough information. If not transient,
+	*  check the new name with the existing path name.
+	*  This is typically transient for schemas being created and non-transient for ones being edited.
+	*/
+	if (CustomizedSchema->GetPackage() != GetTransientPackage())
+	{
+		const FString PackageParentName = FPackageName::GetLongPackagePath(CustomizedSchema->GetPackage()->GetPathName());
+		if (!AssetViewUtils::IsValidPackageForCooking(PackageParentName / InText.ToString(), OutErrorMessage))
+		{
+			return false;
+		}
+	}
+
 	return FName::IsValidXName(InText.ToString(), INVALID_OBJECTNAME_CHARACTERS INVALID_LONGPACKAGE_CHARACTERS, &OutErrorMessage);
 }
 
diff --git a/Engine/Plugins/VirtualProduction/CinematicAssemblyTools/Source/CineAssemblyToolsEditor/Private/CineAssemblySchemaFactory.cpp b/Engine/Plugins/VirtualProduction/CinematicAssemblyTools/Source/CineAssemblyToolsEditor/Private/CineAssemblySchemaFactory.cpp
--- a/Engine/Plugins/VirtualProduction/CinematicAssemblyTools/Source/CineAssemblyToolsEditor/Private/CineAssemblySchemaFactory.cpp
+++ b/Engine/Plugins/VirtualProduction/CinematicAssemblyTools/Source/CineAssemblyToolsEditor/Private/CineAssemblySchemaFactory.cpp
@@ -2,6 +2,7 @@
 
 #include "CineAssemblySchemaFactory.h"
 
+#include "AssetViewUtils.h"
 #include "AssetToolsModule.h"
 #include "CineAssemblySchema.h"
 #include "CineAssemblyToolsAnalytics.h"
@@ -66,6 +67,14 @@ void UCineAssemblySchemaFactory::CreateConfiguredSchema(UCineAssemblySchema* Con
 	}
 
 	const FString DesiredPackageName = CreateAssetPath / ConfiguredSchema->SchemaName;
+
+	// Ensure the package name length does not exceed the maximum cook path length as this may cause issues later on.
+	FText OutErrorMessage;
+	if (!AssetViewUtils::IsValidPackageForCooking(DesiredPackageName, OutErrorMessage))
+	{
+		return;
+	}
+
 	const FString DesiredSuffix = TEXT("");
 
 	IAssetTools& AssetTools = FModuleManager::GetModuleChecked<FAssetToolsModule>("AssetTools").Get();
@@ -74,6 +83,12 @@ void UCineAssemblySchemaFactory::CreateConfiguredSchema(UCineAssemblySchema* Con
 	FString UniqueAssetName;
 	AssetTools.CreateUniqueAssetName(DesiredPackageName, DesiredSuffix, UniquePackageName, UniqueAssetName);
 
+	// Double check again now that we may have created a new unique name
+	if (!AssetViewUtils::IsValidPackageForCooking(UniquePackageName, OutErrorMessage))
+	{
+		return;
+	}
+
 	// The input schema object was created in the transient package while its properties were configured.
 	// Now, we can create a real package for it, rename it, and update its object flags.
 	UPackage* Package = CreatePackage(*UniquePackageName);
diff --git a/Engine/Plugins/VirtualProduction/CinematicAssemblyTools/Source/CineAssemblyToolsEditor/Private/UI/CineAssembly/SCineAssemblyConfigWindow.cpp b/Engine/Plugins/VirtualProduction/CinematicAssemblyTools/Source/CineAssemblyToolsEditor/Private/UI/CineAssembly/SCineAssemblyConfigWindow.cpp
--- a/Engine/Plugins/VirtualProduction/CinematicAssemblyTools/Source/CineAssemblyToolsEditor/Private/UI/CineAssembly/SCineAssemblyConfigWindow.cpp
+++ b/Engine/Plugins/VirtualProduction/CinematicAssemblyTools/Source/CineAssemblyToolsEditor/Private/UI/CineAssembly/SCineAssemblyConfigWindow.cpp
@@ -2,6 +2,7 @@
 
 #include "SCineAssemblyConfigWindow.h"
 
+#include "AssetViewUtils.h"
 #include "CineAssemblyFactory.h"
 #include "CineAssemblyToolsStyle.h"
 #include "ContentBrowserModule.h"
@@ -211,6 +212,18 @@ TSharedRef<SWidget> SCineAssemblyConfigWindow::MakeButtonsPanel()
 						.ButtonStyle(FAppStyle::Get(), "PrimaryButton")
 						.HAlign(HAlign_Center)
 						.OnClicked(this, &SCineAssemblyConfigWindow::OnCreateAssetClicked)
+						.IsEnabled_Lambda([this]() -> bool
+							{
+								FText OutErrorMessage;
+								const FString AssemblyPackageName = UCineAssemblyFactory::MakeAssemblyPackageName(CineAssemblyToConfigure.Get(), CreateAssetPath);
+								return AssetViewUtils::IsValidPackageForCooking(AssemblyPackageName, OutErrorMessage);
+							})
+						.ToolTipText_Lambda([this]() -> FText
+							{
+								FText OutErrorMessage;
+								const FString AssemblyPackageName = UCineAssemblyFactory::MakeAssemblyPackageName(CineAssemblyToConfigure.Get(), CreateAssetPath);
+								return !AssetViewUtils::IsValidPackageForCooking(AssemblyPackageName, OutErrorMessage) ? OutErrorMessage : FText();
+							})
 				]
 
 			+ SHorizontalBox::Slot()
diff --git a/Engine/Plugins/VirtualProduction/CinematicAssemblyTools/Source/CineAssemblyToolsEditor/Private/UI/CineAssembly/SCineAssemblySchemaWindow.cpp b/Engine/Plugins/VirtualProduction/CinematicAssemblyTools/Source/CineAssemblyToolsEditor/Private/UI/CineAssembly/SCineAssemblySchemaWindow.cpp
--- a/Engine/Plugins/VirtualProduction/CinematicAssemblyTools/Source/CineAssemblyToolsEditor/Private/UI/CineAssembly/SCineAssemblySchemaWindow.cpp
+++ b/Engine/Plugins/VirtualProduction/CinematicAssemblyTools/Source/CineAssemblyToolsEditor/Private/UI/CineAssembly/SCineAssemblySchemaWindow.cpp
@@ -219,6 +219,15 @@ bool SCineAssemblySchemaWindow::ValidateSchemaName(const FText& InText, FText& O
 		return false;
 	}
 
+	const FString DesiredPackageName = !CreateAssetPath.IsEmpty()
+		? CreateAssetPath / InText.ToString()
+		: FPackageName::GetLongPackagePath(SchemaToConfigure->GetPackage()->GetPathName()) / InText.ToString();
+
+	if (!AssetViewUtils::IsValidPackageForCooking(DesiredPackageName, OutErrorMessage))
+	{
+		return false;
+	}
+
 	// It is valid if the input text matches the schema's current name
 	if (SchemaToConfigure->SchemaName == InText.ToString())
 	{
@@ -359,6 +368,18 @@ TSharedRef<SWidget> SCineAssemblySchemaWindow::MakeButtonsPanel()
 							.ButtonStyle(FAppStyle::Get(), "PrimaryButton")
 							.HAlign(HAlign_Center)
 							.OnClicked(this, &SCineAssemblySchemaWindow::OnCreateAssetClicked)
+							.IsEnabled_Lambda([this]() -> bool
+								{
+									FText OutErrorMessage;
+									const FString SchemaPackageName = CreateAssetPath / SchemaToConfigure.Get()->SchemaName;
+									return AssetViewUtils::IsValidPackageForCooking(SchemaPackageName, OutErrorMessage);
+								})
+							.ToolTipText_Lambda([this]() -> FText
+								{
+									FText OutErrorMessage;
+									const FString SchemaPackageName = CreateAssetPath / SchemaToConfigure.Get()->SchemaName;
+									return !AssetViewUtils::IsValidPackageForCooking(SchemaPackageName, OutErrorMessage) ? OutErrorMessage : FText();
+								})
 					]
 
 				+ SHorizontalBox::Slot()
