ShotManagement: Expose the Author and Created metadata fields through tokens.

diff --git a/Engine/Plugins/VirtualProduction/CinematicAssemblyTools/Source/CineAssemblyTools/Private/CineAssemblyNamingTokens.cpp b/Engine/Plugins/VirtualProduction/CinematicAssemblyTools/Source/CineAssemblyTools/Private/CineAssemblyNamingTokens.cpp
--- a/Engine/Plugins/VirtualProduction/CinematicAssemblyTools/Source/CineAssemblyTools/Private/CineAssemblyNamingTokens.cpp
+++ b/Engine/Plugins/VirtualProduction/CinematicAssemblyTools/Source/CineAssemblyTools/Private/CineAssemblyNamingTokens.cpp
@@ -2,6 +2,7 @@
 
 #include "CineAssemblyNamingTokens.h"
 
+#include "MovieSceneMetaData.h"
 #include "NamingTokensEngineSubsystem.h"
 
 #include "Algo/IndexOf.h"
@@ -81,6 +82,56 @@ void UCineAssemblyNamingTokens::OnCreateDefaultTokens(TArray<FNamingTokenData>&
 			return FText::GetEmpty();
 		};
 
+
+
+	auto AuthorTokenFunc = [](TWeakObjectPtr<UCineAssembly> InAssembly) -> FText
+		{
+#if WITH_EDITORONLY_DATA
+			UMovieSceneMetaData* MetaData = InAssembly.IsValid() ? InAssembly->FindMetaData<UMovieSceneMetaData>() : nullptr;
+			return MetaData != nullptr 
+				? FText::FromString(MetaData->GetAuthor()) 
+				: FText::GetEmpty();
+#else // WITH_EDITORONLY_DATA
+			return FText::GetEmpty();
+#endif // WITH_EDITORONLY_DATA
+		};
+
+	auto CreatedTokenFunc = [](TWeakObjectPtr<UCineAssembly> InAssembly) -> FText
+		{
+#if WITH_EDITORONLY_DATA
+			UMovieSceneMetaData* MetaData = InAssembly.IsValid() ? InAssembly->FindMetaData<UMovieSceneMetaData>() : nullptr;
+			return MetaData != nullptr 
+				? FText::FromString(MetaData->GetCreated().ToIso8601()) 
+				: FText::GetEmpty();
+#else // WITH_EDITORONLY_DATA
+			return FText::GetEmpty();
+#endif // WITH_EDITORONLY_DATA
+		};
+
+	auto DateCreatedTokenFunc = [](TWeakObjectPtr<UCineAssembly> InAssembly) -> FText
+		{
+#if WITH_EDITORONLY_DATA
+			UMovieSceneMetaData* MetaData = InAssembly.IsValid() ? InAssembly->FindMetaData<UMovieSceneMetaData>() : nullptr;
+			return MetaData != nullptr
+				? FText::FromString(MetaData->GetCreated().ToFormattedString(TEXT("%Y-%m-%d")))
+				: FText::GetEmpty();
+#else // WITH_EDITORONLY_DATA
+			return FText::GetEmpty();
+#endif // WITH_EDITORONLY_DATA
+		};
+
+	auto TimeCreatedTokenFunc = [](TWeakObjectPtr<UCineAssembly> InAssembly) -> FText
+		{
+#if WITH_EDITORONLY_DATA
+			UMovieSceneMetaData* MetaData = InAssembly.IsValid() ? InAssembly->FindMetaData<UMovieSceneMetaData>() : nullptr;
+			return MetaData != nullptr
+				? FText::FromString(MetaData->GetCreated().ToFormattedString(TEXT("%H:%M:%S")))
+				: FText::GetEmpty();
+#else // WITH_EDITORONLY_DATA
+			return FText::GetEmpty();
+#endif // WITH_EDITORONLY_DATA
+		};
+
 	FNamingTokenData AssemblyNameToken;
 	AssemblyNameToken.TokenKey = TEXT("assembly");
 	AssemblyNameToken.DisplayName = LOCTEXT("AssemblyNameToken", "Assembly Name");
@@ -106,11 +157,35 @@ void UCineAssemblyNamingTokens::OnCreateDefaultTokens(TArray<FNamingTokenData>&
 	ProductionToken.DisplayName = LOCTEXT("ProductionTokenName", "Production");
 	ProductionToken.TokenProcessorNative.BindLambda([this, ProductionTokenFunc]() { return ExecuteTokenFunc(ProductionTokenFunc); });
 
+	FNamingTokenData AuthorToken;
+	AuthorToken.TokenKey = TEXT("author");
+	AuthorToken.DisplayName = LOCTEXT("AuthorTokenName", "Author");
+	AuthorToken.TokenProcessorNative.BindLambda([this, AuthorTokenFunc]() { return ExecuteTokenFunc(AuthorTokenFunc); });
+
+	FNamingTokenData CreatedToken;
+	CreatedToken.TokenKey = TEXT("created");
+	CreatedToken.DisplayName = LOCTEXT("CreatedTokenName", "Created");
+	CreatedToken.TokenProcessorNative.BindLambda([this, CreatedTokenFunc]() { return ExecuteTokenFunc(CreatedTokenFunc); });
+
+	FNamingTokenData DateCreatedToken;
+	DateCreatedToken.TokenKey = TEXT("dateCreated");
+	DateCreatedToken.DisplayName = LOCTEXT("DateCreatedTokenName", "Date Created");
+	DateCreatedToken.TokenProcessorNative.BindLambda([this, DateCreatedTokenFunc]() { return ExecuteTokenFunc(DateCreatedTokenFunc); });
+
+	FNamingTokenData TimeCreatedToken;
+	TimeCreatedToken.TokenKey = TEXT("timeCreated");
+	TimeCreatedToken.DisplayName = LOCTEXT("TimeCreatedTokenName", "Time Created");
+	TimeCreatedToken.TokenProcessorNative.BindLambda([this, TimeCreatedTokenFunc]() { return ExecuteTokenFunc(TimeCreatedTokenFunc); });
+
 	Tokens.Add(AssemblyNameToken);
 	Tokens.Add(SchemaToken);
 	Tokens.Add(LevelToken);
 	Tokens.Add(ParentToken);
 	Tokens.Add(ProductionToken);
+	Tokens.Add(AuthorToken);
+	Tokens.Add(CreatedToken);
+	Tokens.Add(DateCreatedToken);
+	Tokens.Add(TimeCreatedToken);
 }
 
 void UCineAssemblyNamingTokens::OnPreEvaluate_Implementation(const FNamingTokensEvaluationData& InEvaluationData)
