ShotManagement: Access CAT created and author metadata through blueprint and python API.

diff --git a/Engine/Plugins/VirtualProduction/CinematicAssemblyTools/Source/CineAssemblyTools/Private/CineAssembly.cpp b/Engine/Plugins/VirtualProduction/CinematicAssemblyTools/Source/CineAssemblyTools/Private/CineAssembly.cpp
--- a/Engine/Plugins/VirtualProduction/CinematicAssemblyTools/Source/CineAssemblyTools/Private/CineAssembly.cpp
+++ b/Engine/Plugins/VirtualProduction/CinematicAssemblyTools/Source/CineAssemblyTools/Private/CineAssembly.cpp
@@ -17,6 +17,7 @@
 
 #if WITH_EDITOR
 #include "AssetToolsModule.h"
+#include "MovieSceneMetaData.h"
 #endif
 
 #define LOCTEXT_NAMESPACE "CineAssembly"
@@ -383,6 +384,41 @@ void UCineAssembly::SetParentAssembly(TSoftObjectPtr<UCineAssembly> InParent)
 	ParentAssembly = InParent.ToSoftObjectPath();
 }
 
+#if WITH_EDITOR
+
+FString UCineAssembly::GetAuthor() const
+{
+	UMovieSceneMetaData* MetaData = FindMetaData<UMovieSceneMetaData>();
+	return MetaData != nullptr ? MetaData->GetAuthor() : FString();
+}
+
+void UCineAssembly::SetAuthor(const FString& InAuthor)
+{
+	UMovieSceneMetaData* MetaData = FindOrAddMetaData<UMovieSceneMetaData>();
+	MetaData->Modify();
+	MetaData->SetAuthor(InAuthor);
+}
+
+FString UCineAssembly::GetCreatedIso8601String() const
+{
+	UMovieSceneMetaData* MetaData = FindMetaData<UMovieSceneMetaData>();
+	return MetaData != nullptr ? MetaData->GetCreated().ToIso8601() : FString();
+}
+
+FString UCineAssembly::GetDateCreatedString() const
+{
+	UMovieSceneMetaData* MetaData = FindMetaData<UMovieSceneMetaData>();
+	return MetaData != nullptr ? MetaData->GetCreated().ToFormattedString(TEXT("%Y-%m-%d")) : FString();
+}
+
+FString UCineAssembly::GetTimeCreatedString() const
+{
+	UMovieSceneMetaData* MetaData = FindMetaData<UMovieSceneMetaData>();
+	return MetaData != nullptr ? MetaData->GetCreated().ToFormattedString(TEXT("%H:%M:%S")) : FString();
+}
+
+#endif // WITH_EDITOR
+
 FString UCineAssembly::GetFullMetadataString() const
 {
 	FString JsonString;
diff --git a/Engine/Plugins/VirtualProduction/CinematicAssemblyTools/Source/CineAssemblyTools/Public/CineAssembly.h b/Engine/Plugins/VirtualProduction/CinematicAssemblyTools/Source/CineAssemblyTools/Public/CineAssembly.h
--- a/Engine/Plugins/VirtualProduction/CinematicAssemblyTools/Source/CineAssemblyTools/Public/CineAssembly.h
+++ b/Engine/Plugins/VirtualProduction/CinematicAssemblyTools/Source/CineAssemblyTools/Public/CineAssembly.h
@@ -93,6 +93,28 @@ public:
 	UFUNCTION(BlueprintCallable, Category = "Cine Assembly Tools")
 	CINEASSEMBLYTOOLS_API void SetParentAssembly(TSoftObjectPtr<UCineAssembly> InParent);
 
+#if WITH_EDITOR
+	/** Gets the author of this assembly */
+	UFUNCTION(BlueprintCallable, Category = "Cine Assembly Tools")
+	CINEASSEMBLYTOOLS_API FString GetAuthor() const;
+
+	/** Set the author of this assembly */
+	UFUNCTION(BlueprintCallable, Category = "Cine Assembly Tools")
+	CINEASSEMBLYTOOLS_API void SetAuthor(const FString& InAuthor);
+
+	/** Gets the created date-time of this assembly as an Iso 8601 format string. */
+	UFUNCTION(BlueprintCallable, Category = "Cine Assembly Tools")
+	CINEASSEMBLYTOOLS_API FString GetCreatedIso8601String() const;
+
+	/** Gets the date this assembly was created as a YYYY-MM-DD string. */
+	UFUNCTION(BlueprintCallable, Category = "Cine Assembly Tools")
+	CINEASSEMBLYTOOLS_API FString GetDateCreatedString() const;
+
+	/** Gets the 24-hour time of day this assembly was created as a HH:MM:SS string. */
+	UFUNCTION(BlueprintCallable, Category = "Cine Assembly Tools")
+	CINEASSEMBLYTOOLS_API FString GetTimeCreatedString() const;
+#endif // WITH_EDITOR
+
 	/** Get all of the metadata for this assembly as a formatted JSON string */
 	UFUNCTION(BlueprintPure, Category = "Cine Assembly Tools")
 	CINEASSEMBLYTOOLS_API FString GetFullMetadataString() const;
