ShotManagement: Adding data-only flag to CineAssembly assets and Cine Assembly Schemas. Opening a DataOnly Assembly opens only the details window and not the full sequence. Creating an assembly from a schema mirrors the schemas Data-Only flag as the initial value on the assembly.

diff --git a/Engine/Plugins/VirtualProduction/CinematicAssemblyTools/Source/CineAssemblyTools/Private/CineAssembly.cpp b/Engine/Plugins/VirtualProduction/CinematicAssemblyTools/Source/CineAssemblyTools/Private/CineAssembly.cpp
--- a/Engine/Plugins/VirtualProduction/CinematicAssemblyTools/Source/CineAssemblyTools/Private/CineAssembly.cpp
+++ b/Engine/Plugins/VirtualProduction/CinematicAssemblyTools/Source/CineAssemblyTools/Private/CineAssembly.cpp
@@ -158,6 +158,12 @@ void UCineAssembly::ChangeSchema(UCineAssemblySchema* InSchema)
 			DefaultFolderNames.Add(NewTemplateName);
 		}
 	}
+
+	// Update the Data Only flag on the assembly
+	if (BaseSchema)
+	{
+		bIsDataOnly = BaseSchema->bIsDataOnly;
+	}
 }
 
 #if WITH_EDITOR
diff --git a/Engine/Plugins/VirtualProduction/CinematicAssemblyTools/Source/CineAssemblyTools/Public/CineAssembly.h b/Engine/Plugins/VirtualProduction/CinematicAssemblyTools/Source/CineAssemblyTools/Public/CineAssembly.h
--- a/Engine/Plugins/VirtualProduction/CinematicAssemblyTools/Source/CineAssemblyTools/Public/CineAssembly.h
+++ b/Engine/Plugins/VirtualProduction/CinematicAssemblyTools/Source/CineAssemblyTools/Public/CineAssembly.h
@@ -205,6 +205,10 @@ public:
 	UPROPERTY()
 	TArray<TObjectPtr<UMovieSceneSubSection>> SubAssemblies;
 
+	/** Mark this CineAssembly as Data-Only. When Data-Only, opening a CineAssembly opens only it's details. */
+	UPROPERTY(AssetRegistrySearchable, BlueprintReadWrite, EditAnywhere, Category="Default", AdvancedDisplay, meta=(DisplayName="Data-Only"))
+	bool bIsDataOnly = false;
+
 	/** The asset registry tag that contains the assembly type information */
 	static CINEASSEMBLYTOOLS_API const FName AssetRegistryTag_AssemblyType;
 	static CINEASSEMBLYTOOLS_API const FName AssemblyGuidPropertyName;
diff --git a/Engine/Plugins/VirtualProduction/CinematicAssemblyTools/Source/CineAssemblyTools/Public/CineAssemblySchema.h b/Engine/Plugins/VirtualProduction/CinematicAssemblyTools/Source/CineAssemblyTools/Public/CineAssemblySchema.h
--- a/Engine/Plugins/VirtualProduction/CinematicAssemblyTools/Source/CineAssemblyTools/Public/CineAssemblySchema.h
+++ b/Engine/Plugins/VirtualProduction/CinematicAssemblyTools/Source/CineAssemblyTools/Public/CineAssemblySchema.h
@@ -125,6 +125,10 @@ public:
 	UPROPERTY()
 	TArray<FString> FoldersToCreate;
 
+	/** Mark this Schema as Data-Only. Assemblies made from this Schema will have their Data-Only defaulted to this value. */
+	UPROPERTY(AssetRegistrySearchable, BlueprintReadWrite, EditAnywhere, Category = "Default", AdvancedDisplay, meta = (DisplayName = "Data-Only"))
+	bool bIsDataOnly = false;
+
 	static CINEASSEMBLYTOOLS_API const FName SchemaGuidPropertyName;
 
 private:
diff --git a/Engine/Plugins/VirtualProduction/CinematicAssemblyTools/Source/CineAssemblyToolsEditor/Private/AssetDefinition_CineAssembly.cpp b/Engine/Plugins/VirtualProduction/CinematicAssemblyTools/Source/CineAssemblyToolsEditor/Private/AssetDefinition_CineAssembly.cpp
--- a/Engine/Plugins/VirtualProduction/CinematicAssemblyTools/Source/CineAssemblyToolsEditor/Private/AssetDefinition_CineAssembly.cpp
+++ b/Engine/Plugins/VirtualProduction/CinematicAssemblyTools/Source/CineAssemblyToolsEditor/Private/AssetDefinition_CineAssembly.cpp
@@ -75,7 +75,37 @@ const FSlateBrush* UAssetDefinition_CineAssembly::GetThumbnailBrush(const FAsset
 
 EAssetCommandResult UAssetDefinition_CineAssembly::OpenAssets(const FAssetOpenArgs& OpenArgs) const
 {
-	return Super::OpenAssets(OpenArgs);
+	TArray<FAssetData> DataOnlyAssetData;
+	TArray<FAssetData> FullAssetData;
+	
+	const FName IsDataOnlyAssetRegistrySearchablePropertyName = GET_MEMBER_NAME_CHECKED(UCineAssembly, bIsDataOnly);
+
+	for (const FAssetData& AssetData : OpenArgs.Assets)
+	{
+		bool bAssemblyIsDataOnly;
+		AssetData.GetTagValue(IsDataOnlyAssetRegistrySearchablePropertyName, bAssemblyIsDataOnly);
+
+		if (bAssemblyIsDataOnly)
+		{
+			DataOnlyAssetData.Add(AssetData);
+		}
+		else
+		{
+			FullAssetData.Add(AssetData);
+		}
+	}
+
+	FAssetOpenArgs DataOnlyOpenArgs(OpenArgs);
+	DataOnlyOpenArgs.Assets = TConstArrayView<FAssetData>(DataOnlyAssetData);
+	FAssetOpenArgs FullAssetOpenArgs(OpenArgs);
+	FullAssetOpenArgs.Assets = TConstArrayView<FAssetData>(FullAssetData);
+
+	if (!DataOnlyOpenArgs.Assets.IsEmpty())
+	{
+		UAssetDefinitionDefault::OpenAssets(DataOnlyOpenArgs);
+	}
+
+	return Super::OpenAssets(FullAssetOpenArgs);
 }
 
 TArray<FAssetData> UAssetDefinition_CineAssembly::PrepareToActivateAssets(const FAssetActivateArgs& ActivateArgs) const
@@ -88,7 +118,13 @@ TArray<FAssetData> UAssetDefinition_CineAssembly::PrepareToActivateAssets(const
 		const FAssetData& CineAssemblyData = ActivateArgs.Assets[0];
 		AssetsToOpen.Add(CineAssemblyData);
 
-		if (UCineAssembly* CineAssembly = Cast<UCineAssembly>(CineAssemblyData.GetAsset()))
+		const FName IsDataOnlyAssetRegistrySearchablePropertyName = GET_MEMBER_NAME_CHECKED(UCineAssembly, bIsDataOnly);
+
+		bool bAssemblyIsDataOnly;
+		CineAssemblyData.GetTagValue(IsDataOnlyAssetRegistrySearchablePropertyName, bAssemblyIsDataOnly);
+
+		UCineAssembly* CineAssembly = Cast<UCineAssembly>(CineAssemblyData.GetAsset());
+		if (!bAssemblyIsDataOnly && CineAssembly != nullptr)
 		{
 			if (CineAssembly->Level.IsValid())
 			{
