ShotManagement: Expose author field on assembly creation and edit through the details customization.

diff --git a/Engine/Plugins/VirtualProduction/CinematicAssemblyTools/Source/CineAssemblyToolsEditor/Private/CineAssemblyCustomization.cpp b/Engine/Plugins/VirtualProduction/CinematicAssemblyTools/Source/CineAssemblyToolsEditor/Private/CineAssemblyCustomization.cpp
--- a/Engine/Plugins/VirtualProduction/CinematicAssemblyTools/Source/CineAssemblyToolsEditor/Private/CineAssemblyCustomization.cpp
+++ b/Engine/Plugins/VirtualProduction/CinematicAssemblyTools/Source/CineAssemblyToolsEditor/Private/CineAssemblyCustomization.cpp
@@ -10,6 +10,7 @@
 #include "DetailLayoutBuilder.h"
 #include "DetailWidgetRow.h"
 #include "Framework/MultiBox/MultiBoxBuilder.h"
+#include "MovieSceneMetaData.h"
 #include "ProductionSettings.h"
 #include "PropertyCustomizationHelpers.h"
 #include "STemplateStringEditableTextBox.h"
@@ -59,6 +60,9 @@ void FCineAssemblyCustomization::CustomizeDefaultCategory(IDetailLayoutBuilder&
 	TSharedRef<IPropertyHandle> ProdctionPropertyHandle = DetailBuilder.GetProperty(GET_MEMBER_NAME_CHECKED(UCineAssembly, Production));
 	IDetailPropertyRow& ProductionPropertyRow = DefaultCategory.AddProperty(ProdctionPropertyHandle);
 
+	UMovieSceneMetaData* MetaData = CustomizedCineAssembly->FindOrAddMetaData<UMovieSceneMetaData>();
+	DefaultCategory.AddExternalObjectProperty({ MetaData }, FName("Author"));
+
 	const UCineAssemblySchema* BaseSchema = CustomizedCineAssembly->GetSchema();
 	const FSoftObjectPath ParentSchema = BaseSchema ? BaseSchema->ParentSchema : FSoftObjectPath();
 
diff --git a/Engine/Plugins/VirtualProduction/CinematicAssemblyTools/Source/CineAssemblyToolsEditor/Private/UI/CineAssembly/SCineAssemblyConfigWindow.cpp b/Engine/Plugins/VirtualProduction/CinematicAssemblyTools/Source/CineAssemblyToolsEditor/Private/UI/CineAssembly/SCineAssemblyConfigWindow.cpp
--- a/Engine/Plugins/VirtualProduction/CinematicAssemblyTools/Source/CineAssemblyToolsEditor/Private/UI/CineAssembly/SCineAssemblyConfigWindow.cpp
+++ b/Engine/Plugins/VirtualProduction/CinematicAssemblyTools/Source/CineAssemblyToolsEditor/Private/UI/CineAssembly/SCineAssemblyConfigWindow.cpp
@@ -11,6 +11,7 @@
 #include "Interfaces/IPluginManager.h"
 #include "Misc/ConfigCacheIni.h"
 #include "Modules/ModuleManager.h"
+#include "MovieSceneMetaData.h"
 #include "ProductionSettings.h"
 #include "UObject/Package.h"
 #include "Settings/ContentBrowserSettings.h"
@@ -28,6 +29,12 @@ void SCineAssemblyConfigWindow::Construct(const FArguments& InArgs, const FStrin
 
 	// Create a new transient CineAssembly to configure in the UI.
 	CineAssemblyToConfigure = TStrongObjectPtr<UCineAssembly>(NewObject<UCineAssembly>(GetTransientPackage(), NAME_None, RF_Transient));
+	if (UMovieSceneMetaData* MetaDataToConfigure = CineAssemblyToConfigure->FindOrAddMetaData<UMovieSceneMetaData>())
+	{
+		// Default the MetaData for the CineAssembly being configured.
+		MetaDataToConfigure->SetCreated(FDateTime::UtcNow());
+		MetaDataToConfigure->SetAuthor(FApp::GetSessionOwner());
+	}
 
 	UWorld* CurrentWorld = GEditor->GetEditorWorldContext().World();
 	CineAssemblyToConfigure->Level = FSoftObjectPath(CurrentWorld);
diff --git a/Engine/Source/Runtime/LevelSequence/Private/LevelSequence.cpp b/Engine/Source/Runtime/LevelSequence/Private/LevelSequence.cpp
--- a/Engine/Source/Runtime/LevelSequence/Private/LevelSequence.cpp
+++ b/Engine/Source/Runtime/LevelSequence/Private/LevelSequence.cpp
@@ -122,7 +122,11 @@ void ULevelSequence::Initialize()
 #if WITH_EDITOR
 	UMovieSceneMetaData* MetaData = FindOrAddMetaData<UMovieSceneMetaData>();
 	MetaData->SetCreated(FDateTime::UtcNow());
-	MetaData->SetAuthor(FApp::GetSessionOwner());
+	// Only set Author if it has not been set already, potentially by a creation wizard.
+	if (MetaData->GetAuthor().IsEmpty())
+	{
+		MetaData->SetAuthor(FApp::GetSessionOwner());
+	}
 #endif
 }
 
