ShotManagement: Update CineAssembly author and created token access to use the CineAssembly access functions as these are now available.

#rnx

diff --git a/Engine/Plugins/VirtualProduction/CinematicAssemblyTools/Source/CineAssemblyTools/Private/CineAssemblyNamingTokens.cpp b/Engine/Plugins/VirtualProduction/CinematicAssemblyTools/Source/CineAssemblyTools/Private/CineAssemblyNamingTokens.cpp
--- a/Engine/Plugins/VirtualProduction/CinematicAssemblyTools/Source/CineAssemblyTools/Private/CineAssemblyNamingTokens.cpp
+++ b/Engine/Plugins/VirtualProduction/CinematicAssemblyTools/Source/CineAssemblyTools/Private/CineAssemblyNamingTokens.cpp
@@ -86,46 +86,42 @@ void UCineAssemblyNamingTokens::OnCreateDefaultTokens(TArray<FNamingTokenData>&
 
 	auto AuthorTokenFunc = [](TWeakObjectPtr<UCineAssembly> InAssembly) -> FText
 		{
-#if WITH_EDITORONLY_DATA
-			UMovieSceneMetaData* MetaData = InAssembly.IsValid() ? InAssembly->FindMetaData<UMovieSceneMetaData>() : nullptr;
-			return MetaData != nullptr 
-				? FText::FromString(MetaData->GetAuthor()) 
+#if WITH_EDITOR
+			return InAssembly.IsValid()
+				? FText::FromString(InAssembly->GetAuthor()) 
 				: FText::GetEmpty();
-#else // WITH_EDITORONLY_DATA
+#else // WITH_EDITOR
 			return FText::GetEmpty();
-#endif // WITH_EDITORONLY_DATA
+#endif // WITH_EDITOR
 		};
 
 	auto CreatedTokenFunc = [](TWeakObjectPtr<UCineAssembly> InAssembly) -> FText
 		{
-#if WITH_EDITORONLY_DATA
-			UMovieSceneMetaData* MetaData = InAssembly.IsValid() ? InAssembly->FindMetaData<UMovieSceneMetaData>() : nullptr;
-			return MetaData != nullptr 
-				? FText::FromString(MetaData->GetCreated().ToIso8601()) 
+#if WITH_EDITOR
+			return InAssembly.IsValid()
+				? FText::FromString(InAssembly->GetCreatedIso8601String()) 
 				: FText::GetEmpty();
-#else // WITH_EDITORONLY_DATA
+#else // WITH_EDITOR
 			return FText::GetEmpty();
-#endif // WITH_EDITORONLY_DATA
+#endif // WITH_EDITOR
 		};
 
 	auto DateCreatedTokenFunc = [](TWeakObjectPtr<UCineAssembly> InAssembly) -> FText
 		{
-#if WITH_EDITORONLY_DATA
-			UMovieSceneMetaData* MetaData = InAssembly.IsValid() ? InAssembly->FindMetaData<UMovieSceneMetaData>() : nullptr;
-			return MetaData != nullptr
-				? FText::FromString(MetaData->GetCreated().ToFormattedString(TEXT("%Y-%m-%d")))
+#if WITH_EDITOR
+			return InAssembly.IsValid()
+				? FText::FromString(InAssembly->GetDateCreatedString())
 				: FText::GetEmpty();
-#else // WITH_EDITORONLY_DATA
+#else // WITH_EDITOR
 			return FText::GetEmpty();
-#endif // WITH_EDITORONLY_DATA
+#endif // WITH_EDITOR
 		};
 
 	auto TimeCreatedTokenFunc = [](TWeakObjectPtr<UCineAssembly> InAssembly) -> FText
 		{
 #if WITH_EDITORONLY_DATA
-			UMovieSceneMetaData* MetaData = InAssembly.IsValid() ? InAssembly->FindMetaData<UMovieSceneMetaData>() : nullptr;
-			return MetaData != nullptr
-				? FText::FromString(MetaData->GetCreated().ToFormattedString(TEXT("%H:%M:%S")))
+			return InAssembly.IsValid()
+				? FText::FromString(InAssembly->GetTimeCreatedString())
 				: FText::GetEmpty();
 #else // WITH_EDITORONLY_DATA
 			return FText::GetEmpty();
