ShotManagement: Mark CineAssemblySchema as BlueprintType.

diff --git a/Engine/Plugins/VirtualProduction/CinematicAssemblyTools/Source/CineAssemblyTools/Public/CineAssemblySchema.h b/Engine/Plugins/VirtualProduction/CinematicAssemblyTools/Source/CineAssemblyTools/Public/CineAssemblySchema.h
--- a/Engine/Plugins/VirtualProduction/CinematicAssemblyTools/Source/CineAssemblyTools/Public/CineAssemblySchema.h
+++ b/Engine/Plugins/VirtualProduction/CinematicAssemblyTools/Source/CineAssemblyTools/Public/CineAssemblySchema.h
@@ -48,7 +48,7 @@ struct FAssemblyMetadataDesc
 /**
  * A template object for building different Cine Assembly types
  */
-UCLASS(MinimalAPI)
+UCLASS(MinimalAPI, BlueprintType)
 class UCineAssemblySchema : public UObject
 {
 	GENERATED_BODY()
