ShotManagement: Adding support of using token text in CineAssembly metadata values.

diff --git a/Engine/Plugins/VirtualProduction/CinematicAssemblyTools/Source/CineAssemblyTools/Private/CineAssembly.cpp b/Engine/Plugins/VirtualProduction/CinematicAssemblyTools/Source/CineAssemblyTools/Private/CineAssembly.cpp
--- a/Engine/Plugins/VirtualProduction/CinematicAssemblyTools/Source/CineAssemblyTools/Private/CineAssembly.cpp
+++ b/Engine/Plugins/VirtualProduction/CinematicAssemblyTools/Source/CineAssemblyTools/Private/CineAssembly.cpp
@@ -464,12 +464,24 @@ void UCineAssembly::SetMetadataAsFloat(FString InKey, float InValue)
 }
 
 bool UCineAssembly::GetMetadataAsString(FString InKey, FString& OutValue) const
+{
+	return GetMetadataAsString(InKey, OutValue, true /* bEvaluateTokens */);
+}
+
+bool UCineAssembly::GetMetadataAsString(FString InKey, FString& OutValue, bool bEvaluateTokens) const
 {
 	if (!MetadataJsonObject->TryGetStringField(InKey, OutValue))
 	{
 		OutValue = TEXT("");
 		return false;
 	}
+
+	if (bEvaluateTokens)
+	{
+		// Resolve any tokens in the returned string
+		OutValue = UCineAssemblyNamingTokens::GetResolvedText(OutValue, const_cast<UCineAssembly*>(this)).ToString();
+	}
+
 	return true;
 }
 
diff --git a/Engine/Plugins/VirtualProduction/CinematicAssemblyTools/Source/CineAssemblyTools/Public/CineAssembly.h b/Engine/Plugins/VirtualProduction/CinematicAssemblyTools/Source/CineAssemblyTools/Public/CineAssembly.h
--- a/Engine/Plugins/VirtualProduction/CinematicAssemblyTools/Source/CineAssemblyTools/Public/CineAssembly.h
+++ b/Engine/Plugins/VirtualProduction/CinematicAssemblyTools/Source/CineAssemblyTools/Public/CineAssembly.h
@@ -135,10 +135,13 @@ public:
 	UFUNCTION(BlueprintCallable, Category = "Cine Assembly Tools")
 	CINEASSEMBLYTOOLS_API void SetMetadataAsFloat(FString InKey, float InValue);
 
-	/** Get the metadata value for the input key as a string (if it exists) */
+	/** Get the metadata value for the input key as a token evaluated string (if it exists) */
 	UFUNCTION(BlueprintPure, Category = "Cine Assembly Tools", meta = (ReturnDisplayName = "FoundKey"))
 	CINEASSEMBLYTOOLS_API bool GetMetadataAsString(FString InKey, FString& OutValue) const;
 
+	/** Get the metadata value for the input key as a string (if it exists), with tokens evaluated if requested. */
+	CINEASSEMBLYTOOLS_API bool GetMetadataAsString(FString InKey, FString& OutValue, bool bEvaluateTokens) const;
+
 	/** Get the metadata value for the input key as a boolean (if it exists) */
 	UFUNCTION(BlueprintPure, Category = "Cine Assembly Tools", meta = (ReturnDisplayName = "FoundKey"))
 	CINEASSEMBLYTOOLS_API bool GetMetadataAsBool(FString InKey, bool& OutValue) const;
diff --git a/Engine/Plugins/VirtualProduction/CinematicAssemblyTools/Source/CineAssemblyTools/Public/CineAssemblySchema.h b/Engine/Plugins/VirtualProduction/CinematicAssemblyTools/Source/CineAssemblyTools/Public/CineAssemblySchema.h
--- a/Engine/Plugins/VirtualProduction/CinematicAssemblyTools/Source/CineAssemblyTools/Public/CineAssemblySchema.h
+++ b/Engine/Plugins/VirtualProduction/CinematicAssemblyTools/Source/CineAssemblyTools/Public/CineAssemblySchema.h
@@ -43,6 +43,10 @@ struct FAssemblyMetadataDesc
 
 	/** The default value for this metadata field */
 	TVariant<FString, bool, int32, float> DefaultValue;
+
+	/** If string field types should be evaluated as template token strings. */
+	UPROPERTY(EditAnywhere, Category = "Metadata", meta = (EditCondition = "Type == ECineAssemblyMetadataType::String", EditConditionHides))
+	bool bEvaluateTokens = false;
 };
 
 /**
diff --git a/Engine/Plugins/VirtualProduction/CinematicAssemblyTools/Source/CineAssemblyToolsEditor/Private/CineAssemblyCustomization.cpp b/Engine/Plugins/VirtualProduction/CinematicAssemblyTools/Source/CineAssemblyToolsEditor/Private/CineAssemblyCustomization.cpp
--- a/Engine/Plugins/VirtualProduction/CinematicAssemblyTools/Source/CineAssemblyToolsEditor/Private/CineAssemblyCustomization.cpp
+++ b/Engine/Plugins/VirtualProduction/CinematicAssemblyTools/Source/CineAssemblyToolsEditor/Private/CineAssemblyCustomization.cpp
@@ -141,22 +141,7 @@ void FCineAssemblyCustomization::CustomizeMetadataCategory(IDetailLayoutBuilder&
 			ValueWidget = SNew(SBox)
 				.MaxDesiredHeight(120.0f)
 				[
-					SNew(SMultiLineEditableTextBox)
-						.Font(FAppStyle::GetFontStyle(TEXT("PropertyWindow.NormalFont")))
-						.AutoWrapText(true)
-						.Text_Lambda([this, MetadataDesc]()
-							{
-								FString Value;
-								if (!CustomizedCineAssembly->GetMetadataAsString(MetadataDesc.Key, Value))
-								{
-									Value = MetadataDesc.DefaultValue.Get<FString>();
-								}
-								return FText::FromString(Value);
-							})
-						.OnTextCommitted_Lambda([this, MetadataDesc](const FText& InText, ETextCommit::Type InCommitType)
-							{
-								CustomizedCineAssembly->SetMetadataAsString(MetadataDesc.Key, InText.ToString());
-							})
+					MakeStringValueWidget(MetadataDesc)
 				];
 		}
 		else if (MetadataDesc.Type == ECineAssemblyMetadataType::Bool)
@@ -318,6 +303,60 @@ void FCineAssemblyCustomization::CustomizeSubsequenceCategory(IDetailLayoutBuild
 	}
 }
 
+TSharedRef<SWidget> FCineAssemblyCustomization::MakeStringValueWidget(const FAssemblyMetadataDesc& MetadataDesc)
+{
+	TSharedPtr<SWidget> StringValueWidget;
+	if (!MetadataDesc.bEvaluateTokens)
+	{
+		StringValueWidget = SNew(SMultiLineEditableTextBox)
+			.Font(FAppStyle::GetFontStyle(TEXT("PropertyWindow.NormalFont")))
+			.AutoWrapText(true)
+			.Text_Lambda([this, MetadataDesc]()
+				{
+					FString Value;
+					if (!CustomizedCineAssembly->GetMetadataAsString(MetadataDesc.Key, Value))
+					{
+						Value = MetadataDesc.DefaultValue.Get<FString>();
+					}
+					return FText::FromString(Value);
+				})
+			.OnTextCommitted_Lambda([this, MetadataDesc](const FText& InText, ETextCommit::Type InCommitType)
+				{
+					CustomizedCineAssembly->SetMetadataAsString(MetadataDesc.Key, InText.ToString());
+				});
+	}
+	else
+	{
+		StringValueWidget = SNew(STemplateStringEditableTextBox)
+			.AllowMultiLine(true)
+			.Text_Lambda([this, MetadataDesc]() -> FText
+				{
+					FString Value;
+					const bool bEvaluateTokens = false;
+					if (!CustomizedCineAssembly->GetMetadataAsString(MetadataDesc.Key, Value, bEvaluateTokens))
+					{
+						Value = MetadataDesc.DefaultValue.Get<FString>();
+					}
+					return FText::FromString(Value);
+				})
+			.ResolvedText_Lambda([this, MetadataDesc]()
+				{
+					FString Value;
+					if (!CustomizedCineAssembly->GetMetadataAsString(MetadataDesc.Key, Value))
+					{
+						Value = MetadataDesc.DefaultValue.Get<FString>();
+					}
+					return FText::FromString(Value);
+				})
+			.OnTextCommitted_Lambda([this, MetadataDesc](const FText& InText, ETextCommit::Type InCommitType)
+				{
+					CustomizedCineAssembly->SetMetadataAsString(MetadataDesc.Key, InText.ToString());
+				});
+	}
+
+	return StringValueWidget.ToSharedRef();
+}
+
 bool FCineAssemblyCustomization::ShouldFilterAssetBySchema(const FAssetData& InAssetData, FSoftObjectPath Schema)
 {
 	if (Schema.IsValid())
diff --git a/Engine/Plugins/VirtualProduction/CinematicAssemblyTools/Source/CineAssemblyToolsEditor/Private/CineAssemblyCustomization.h b/Engine/Plugins/VirtualProduction/CinematicAssemblyTools/Source/CineAssemblyToolsEditor/Private/CineAssemblyCustomization.h
--- a/Engine/Plugins/VirtualProduction/CinematicAssemblyTools/Source/CineAssemblyToolsEditor/Private/CineAssemblyCustomization.h
+++ b/Engine/Plugins/VirtualProduction/CinematicAssemblyTools/Source/CineAssemblyToolsEditor/Private/CineAssemblyCustomization.h
@@ -8,6 +8,7 @@
 #include "UObject/TemplateString.h"
 
 class UCineAssembly;
+struct FAssemblyMetadataDesc;
 
 /**
  * Detail customization for UCineAssembly
@@ -28,6 +29,12 @@ private:
 	void CustomizeMetadataCategory(IDetailLayoutBuilder& DetailBuilder);
 	void CustomizeSubsequenceCategory(IDetailLayoutBuilder& DetailBuilder);
 
+	/** 
+	* Given a FAssemblyMetadataDesc, makes a STemplateStringEditableTextBox for a string value that has bEvaluateTokens set to true,
+	* and a SMultiLineEditableTextBox for a string value that has bEvaluateTokens set to false.
+	*/
+	TSharedRef<SWidget> MakeStringValueWidget(const FAssemblyMetadataDesc& MetadataDesc);
+
 	/**
 	 * Evaluates all of the customization's template strings with the naming tokens subsystem.
 	 * This function is throttled to only run at a set frequency, to avoid the potential to constantly query the naming tokens subsystem.
diff --git a/Engine/Plugins/VirtualProduction/CinematicAssemblyTools/Source/CineAssemblyToolsEditor/Private/CineAssemblyMetadataCustomization.cpp b/Engine/Plugins/VirtualProduction/CinematicAssemblyTools/Source/CineAssemblyToolsEditor/Private/CineAssemblyMetadataCustomization.cpp
--- a/Engine/Plugins/VirtualProduction/CinematicAssemblyTools/Source/CineAssemblyToolsEditor/Private/CineAssemblyMetadataCustomization.cpp
+++ b/Engine/Plugins/VirtualProduction/CinematicAssemblyTools/Source/CineAssemblyToolsEditor/Private/CineAssemblyMetadataCustomization.cpp
@@ -10,6 +10,7 @@
 #include "IPropertyUtilities.h"
 #include "NamingTokensEngineSubsystem.h"
 #include "PropertyCustomizationHelpers.h"
+#include "STemplateStringEditableTextBox.h"
 #include "Widgets/Input/SCheckBox.h"
 #include "Widgets/Input/SEditableTextBox.h"
 #include "Widgets/Input/SMultiLineEditableTextBox.h"
@@ -37,6 +38,73 @@ void FCineAssemblyMetadataCustomization::CustomizeHeader(TSharedRef<IPropertyHan
 		];
 }
 
+TSharedRef<SWidget> FCineAssemblyMetadataCustomization::MakeStringDefaultValueWidget(FAssemblyMetadataDesc& MetadataDesc, TSharedRef<IPropertyHandle> PropertyHandle)
+{
+	TSharedPtr<IPropertyHandle> MetadataEvaluateTokensHandle = PropertyHandle->GetChildHandle(GET_MEMBER_NAME_CHECKED(FAssemblyMetadataDesc, bEvaluateTokens));
+
+	TSharedPtr<SWidgetSwitcher> StringValueWidget = SNew(SWidgetSwitcher)
+		.WidgetIndex_Lambda([MetadataEvaluateTokensHandle]()
+			{
+				bool bValue;
+				MetadataEvaluateTokensHandle->GetValue(bValue);
+				return static_cast<int32>(bValue);
+			});
+
+	const int32 EvalFalseIndex = 0;
+	const int32 EvalTrueIndex = 1;
+
+	StringValueWidget->AddSlot(EvalFalseIndex)
+		[
+			SNew(SMultiLineEditableTextBox)
+				.Font(FAppStyle::GetFontStyle(TEXT("PropertyWindow.NormalFont")))
+				.AutoWrapText(true)
+				.Text_Lambda([&MetadataDesc]() -> FText
+					{
+						if (const FString* Value = MetadataDesc.DefaultValue.TryGet<FString>())
+						{
+							return FText::FromString(*Value);
+						}
+						return FText::GetEmpty();
+					})
+				.OnTextCommitted_Lambda([&MetadataDesc](const FText& InText, ETextCommit::Type InCommitType)
+					{
+						MetadataDesc.DefaultValue.Set<FString>(InText.ToString());
+					})
+		];
+
+	StringValueWidget->AddSlot(EvalTrueIndex)
+		[
+			SNew(STemplateStringEditableTextBox)
+				.AllowMultiLine(true)
+				.Text_Lambda([&MetadataDesc]() -> FText
+					{
+						if (const FString* Value = MetadataDesc.DefaultValue.TryGet<FString>())
+						{
+							return FText::FromString(*Value);
+						}
+						return FText::GetEmpty();
+					})
+				.ResolvedText_Lambda([&MetadataDesc]()
+					{
+						UNamingTokensEngineSubsystem* NamingTokensSubsystem = GEngine->GetEngineSubsystem<UNamingTokensEngineSubsystem>();
+						FNamingTokenFilterArgs FilterArgs;
+						FNamingTokenResultData Result = NamingTokensSubsystem->EvaluateTokenString(MetadataDesc.DefaultValue.Get<FString>(), FilterArgs, {});
+						return Result.EvaluatedText;
+					})
+				.OnTextCommitted_Lambda([&MetadataDesc](const FText& InText, ETextCommit::Type InCommitType)
+					{
+						MetadataDesc.DefaultValue.Set<FString>(InText.ToString());
+					})
+		];
+
+	return SNew(SBox)
+		.MaxDesiredHeight(120.0f)
+		.MinDesiredWidth(400)
+		[
+			StringValueWidget.ToSharedRef()
+		];
+}
+
 void FCineAssemblyMetadataCustomization::CustomizeChildren(TSharedRef<IPropertyHandle> PropertyHandle, IDetailChildrenBuilder& ChildBuilder, IPropertyTypeCustomizationUtils& CustomizationUtils)
 {
 	// Get the schema object that owns the metadata struct being customized
@@ -116,25 +184,7 @@ void FCineAssemblyMetadataCustomization::CustomizeChildren(TSharedRef<IPropertyH
 
 		DefaultValueWidget->AddSlot((uint8)ECineAssemblyMetadataType::String)
 			[
-				SNew(SBox)
-					.MaxDesiredHeight(120.0f)
-					[
-						SNew(SMultiLineEditableTextBox)
-							.Font(FAppStyle::GetFontStyle(TEXT("PropertyWindow.NormalFont")))
-							.AutoWrapText(true)
-							.Text_Lambda([&MetadataDesc]() -> FText
-								{
-									if (const FString* Value = MetadataDesc.DefaultValue.TryGet<FString>())
-									{
-										return FText::FromString(*Value);
-									}
-									return FText::GetEmpty();
-								})
-							.OnTextCommitted_Lambda([&MetadataDesc](const FText& InText, ETextCommit::Type InCommitType)
-								{
-									MetadataDesc.DefaultValue.Set<FString>(InText.ToString());
-								})
-					]
+				MakeStringDefaultValueWidget(MetadataDesc, PropertyHandle)
 			];
 
 		DefaultValueWidget->AddSlot((uint8)ECineAssemblyMetadataType::Bool)
diff --git a/Engine/Plugins/VirtualProduction/CinematicAssemblyTools/Source/CineAssemblyToolsEditor/Private/CineAssemblyMetadataCustomization.h b/Engine/Plugins/VirtualProduction/CinematicAssemblyTools/Source/CineAssemblyToolsEditor/Private/CineAssemblyMetadataCustomization.h
--- a/Engine/Plugins/VirtualProduction/CinematicAssemblyTools/Source/CineAssemblyToolsEditor/Private/CineAssemblyMetadataCustomization.h
+++ b/Engine/Plugins/VirtualProduction/CinematicAssemblyTools/Source/CineAssemblyToolsEditor/Private/CineAssemblyMetadataCustomization.h
@@ -6,6 +6,8 @@
 
 #include "CineAssemblySchema.h"
 
+class SWidget;
+
 class FCineAssemblyMetadataCustomization : public IPropertyTypeCustomization
 {
 public:
@@ -17,6 +19,12 @@ public:
 	//~ End IDetailCustomization interface
 
 private:
+	/**
+	* Given a FAssemblyMetadataDesc, makes a STemplateStringEditableTextBox for a string default value that has bEvaluateTokens set to true,
+	* and a SMultiLineEditableTextBox for a string default value that has bEvaluateTokens set to false.
+	*/
+	TSharedRef<SWidget> MakeStringDefaultValueWidget(FAssemblyMetadataDesc& MetadataDesc, TSharedRef<IPropertyHandle> PropertyHandle);
+
 	/** Customizes the metadata desc Key property */
 	void CustomizeKeyProperty(TSharedRef<IPropertyHandle> PropertyHandle, IDetailChildrenBuilder& ChildBuilder, FAssemblyMetadataDesc& MetadataDesc);
 
