ShotManagement: Use local time instead of UTC when returning created from a CineAssembly as strings.

diff --git a/Engine/Plugins/VirtualProduction/CinematicAssemblyTools/Source/CineAssemblyTools/Private/CineAssembly.cpp b/Engine/Plugins/VirtualProduction/CinematicAssemblyTools/Source/CineAssemblyTools/Private/CineAssembly.cpp
--- a/Engine/Plugins/VirtualProduction/CinematicAssemblyTools/Source/CineAssemblyTools/Private/CineAssembly.cpp
+++ b/Engine/Plugins/VirtualProduction/CinematicAssemblyTools/Source/CineAssemblyTools/Private/CineAssembly.cpp
@@ -25,6 +25,61 @@
 const FName UCineAssembly::AssetRegistryTag_AssemblyType = "AssemblyType";
 const FName UCineAssembly::AssemblyGuidPropertyName = GET_MEMBER_NAME_CHECKED(UCineAssembly, AssemblyGuid);
 
+namespace UE::CineAssembly
+{
+	int32 GetLocalTimezone()
+	{
+		const FTimespan DifferentLocalToUTC = FDateTime::Now() - FDateTime::UtcNow();
+		const int32 DifferenceMinutes = FMath::RoundToInt(DifferentLocalToUTC.GetTotalMinutes());
+
+		const int32 Hours = DifferenceMinutes / 60;
+		const int32 Minutes = DifferenceMinutes % 60;
+
+		const int32 Timezone = (Hours * 100) + Minutes;
+		return Timezone;
+	}
+
+	// Taken from FDateTimeStructCustomization
+	int32 ConvertShortTimezone(int32 ShortTimezone)
+	{
+		// Convert timezones from short-format into long format, -5 -> -0500
+		// Timezone Hour ranges go from -12 to +14 from UTC
+		if (ShortTimezone >= -12 && ShortTimezone <= 14)
+		{
+			return ShortTimezone * 100;
+		}
+
+		// Not a short-form timezone
+		return ShortTimezone;
+	}
+
+	// Taken from FDateTimeStructCustomization
+	FDateTime ConvertTime(const FDateTime& InDate, int32 InTimezone, int32 OutTimezone)
+	{
+		if (InTimezone == OutTimezone)
+		{
+			return InDate;
+		}
+
+		// Timezone Hour ranges go from -12 to +14 from UTC
+		// Convert from whole-hour to the full-format HHMM (-5 -> -0500, 0 -> +0000, etc)
+		InTimezone = ConvertShortTimezone(InTimezone);
+		OutTimezone = ConvertShortTimezone(OutTimezone);
+
+		// Extract timezone minutes
+		const int32 InTimezoneMinutes = (FMath::Abs(InTimezone) % 100);
+		const int32 OutTimezoneMinutes = (FMath::Abs(OutTimezone) % 100);
+
+		// Calculate our Minutes difference
+		const int32 MinutesDifference = OutTimezoneMinutes - InTimezoneMinutes;
+
+		// Calculate our Hours difference
+		const int32 HoursDifference = (OutTimezone / 100) - (InTimezone / 100);
+
+		return FDateTime(InDate + FTimespan(HoursDifference, MinutesDifference, 0));
+	}
+} // namespace UE::CineAssembly
+
 UCineAssembly::UCineAssembly()
 {
 	MetadataJsonObject = MakeShared<FJsonObject>();
@@ -407,20 +462,32 @@ void UCineAssembly::SetAuthor(const FString& InAuthor)
 
 FString UCineAssembly::GetCreatedIso8601String() const
 {
-	UMovieSceneMetaData* MetaData = FindMetaData<UMovieSceneMetaData>();
-	return MetaData != nullptr ? MetaData->GetCreated().ToIso8601() : FString();
+	if (TOptional<FDateTime> CreatedLocalTime = TryGetCreatedAsLocalTime(); CreatedLocalTime.IsSet())
+	{
+		return CreatedLocalTime->ToIso8601();
+	}
+	
+	return FString();
 }
 
 FString UCineAssembly::GetDateCreatedString() const
 {
-	UMovieSceneMetaData* MetaData = FindMetaData<UMovieSceneMetaData>();
-	return MetaData != nullptr ? MetaData->GetCreated().ToFormattedString(TEXT("%Y-%m-%d")) : FString();
+	if (TOptional<FDateTime> CreatedLocalTime = TryGetCreatedAsLocalTime(); CreatedLocalTime.IsSet())
+	{
+		return CreatedLocalTime->ToFormattedString(TEXT("%Y-%m-%d"));
+	}
+
+	return FString();
 }
 
 FString UCineAssembly::GetTimeCreatedString() const
 {
-	UMovieSceneMetaData* MetaData = FindMetaData<UMovieSceneMetaData>();
-	return MetaData != nullptr ? MetaData->GetCreated().ToFormattedString(TEXT("%H:%M:%S")) : FString();
+	if (TOptional<FDateTime> CreatedLocalTime = TryGetCreatedAsLocalTime(); CreatedLocalTime.IsSet())
+	{
+		return CreatedLocalTime->ToFormattedString(TEXT("%H:%M:%S"));
+	}
+
+	return FString();
 }
 
 #endif // WITH_EDITOR
@@ -545,6 +612,19 @@ void UCineAssembly::UpdateInstanceMetadata()
 	}
 }
 
+TOptional<FDateTime> UCineAssembly::TryGetCreatedAsLocalTime() const
+{
+#if WITH_EDITOR
+	// Assuming that the Created field from the Metadata is in UTC.
+	if (UMovieSceneMetaData* MetaData = FindMetaData<UMovieSceneMetaData>(); MetaData != nullptr)
+	{
+		return UE::CineAssembly::ConvertTime(MetaData->GetCreated(), 0 /* UTC */, UE::CineAssembly::GetLocalTimezone());
+	}
+#endif // WITH_EDITOR
+
+	return TOptional<FDateTime>();
+}
+
 void UCineAssembly::AddMetadataNamingToken(const FString& InKey)
 {
 	UNamingTokensEngineSubsystem* NamingTokensSubsystem = GEngine->GetEngineSubsystem<UNamingTokensEngineSubsystem>();
diff --git a/Engine/Plugins/VirtualProduction/CinematicAssemblyTools/Source/CineAssemblyTools/Public/CineAssembly.h b/Engine/Plugins/VirtualProduction/CinematicAssemblyTools/Source/CineAssemblyTools/Public/CineAssembly.h
--- a/Engine/Plugins/VirtualProduction/CinematicAssemblyTools/Source/CineAssemblyTools/Public/CineAssembly.h
+++ b/Engine/Plugins/VirtualProduction/CinematicAssemblyTools/Source/CineAssemblyTools/Public/CineAssembly.h
@@ -167,6 +167,9 @@ private:
 	/** Update the underlying json object whenever keys/values in the instance metadata map are added/removed/modified */
 	void UpdateInstanceMetadata();
 
+	/** Get the created FDateTime in local time. If no metadata found, returns false. */
+	TOptional<FDateTime> TryGetCreatedAsLocalTime() const;
+
 public:
 	/** The assembly name, which supports tokens */
 	UPROPERTY()
