ShotManagement: Corrected private function doc for return type on CineAssembly.

#rnx

diff --git a/Engine/Plugins/VirtualProduction/CinematicAssemblyTools/Source/CineAssemblyTools/Public/CineAssembly.h b/Engine/Plugins/VirtualProduction/CinematicAssemblyTools/Source/CineAssemblyTools/Public/CineAssembly.h
--- a/Engine/Plugins/VirtualProduction/CinematicAssemblyTools/Source/CineAssemblyTools/Public/CineAssembly.h
+++ b/Engine/Plugins/VirtualProduction/CinematicAssemblyTools/Source/CineAssemblyTools/Public/CineAssembly.h
@@ -167,7 +167,7 @@ private:
 	/** Update the underlying json object whenever keys/values in the instance metadata map are added/removed/modified */
 	void UpdateInstanceMetadata();
 
-	/** Get the created FDateTime in local time. If no metadata found, returns false. */
+	/** Get the created FDateTime in local time. If no metadata found, returns an unset optional. */
 	TOptional<FDateTime> TryGetCreatedAsLocalTime() const;
 
 public:
