CineAssembly: Prevent naming assets with a forward slash.

diff --git a/Engine/Plugins/VirtualProduction/CinematicAssemblyTools/Source/CineAssemblyToolsEditor/Private/UI/CineAssembly/SCineAssemblyConfigWindow.cpp b/Engine/Plugins/VirtualProduction/CinematicAssemblyTools/Source/CineAssemblyToolsEditor/Private/UI/CineAssembly/SCineAssemblyConfigWindow.cpp
--- a/Engine/Plugins/VirtualProduction/CinematicAssemblyTools/Source/CineAssemblyToolsEditor/Private/UI/CineAssembly/SCineAssemblyConfigWindow.cpp
+++ b/Engine/Plugins/VirtualProduction/CinematicAssemblyTools/Source/CineAssemblyToolsEditor/Private/UI/CineAssembly/SCineAssemblyConfigWindow.cpp
@@ -200,6 +200,23 @@ void SCineAssemblyConfigWindow::OnSchemaSelected(const FAssetData& AssetData)
 	}
 }
 
+bool SCineAssemblyConfigWindow::ValidateAssetName(FText& OutErrorMessage) const
+{
+	if (CineAssemblyToConfigure.IsValid() && CineAssemblyToConfigure.Get()->AssemblyName.Resolved.ToString().Contains(TEXT("/")))
+	{
+		OutErrorMessage = LOCTEXT("InvalidAssemblyNameErrorMessage", "Assembly names cannot contain '/'");
+		return false;
+	}
+	const FString AssemblyPackageName = UCineAssemblyFactory::MakeAssemblyPackageName(CineAssemblyToConfigure.Get(), CreateAssetPath);
+	return AssetViewUtils::IsValidPackageForCooking(AssemblyPackageName, OutErrorMessage);
+}
+
+bool SCineAssemblyConfigWindow::ValidateAssetName() const
+{
+	FText OutErrorMessage;
+	return ValidateAssetName(OutErrorMessage);
+}
+
 TSharedRef<SWidget> SCineAssemblyConfigWindow::MakeButtonsPanel()
 {
 	return SNew(SBorder)
@@ -219,17 +236,11 @@ TSharedRef<SWidget> SCineAssemblyConfigWindow::MakeButtonsPanel()
 						.ButtonStyle(FAppStyle::Get(), "PrimaryButton")
 						.HAlign(HAlign_Center)
 						.OnClicked(this, &SCineAssemblyConfigWindow::OnCreateAssetClicked)
-						.IsEnabled_Lambda([this]() -> bool
-							{
-								FText OutErrorMessage;
-								const FString AssemblyPackageName = UCineAssemblyFactory::MakeAssemblyPackageName(CineAssemblyToConfigure.Get(), CreateAssetPath);
-								return AssetViewUtils::IsValidPackageForCooking(AssemblyPackageName, OutErrorMessage);
-							})
+						.IsEnabled_Raw(this, &SCineAssemblyConfigWindow::ValidateAssetName)
 						.ToolTipText_Lambda([this]() -> FText
 							{
 								FText OutErrorMessage;
-								const FString AssemblyPackageName = UCineAssemblyFactory::MakeAssemblyPackageName(CineAssemblyToConfigure.Get(), CreateAssetPath);
-								return !AssetViewUtils::IsValidPackageForCooking(AssemblyPackageName, OutErrorMessage) ? OutErrorMessage : FText();
+								return !ValidateAssetName(OutErrorMessage) ? OutErrorMessage : FText();
 							})
 				]
 
@@ -260,9 +271,12 @@ FText SCineAssemblyConfigWindow::GetCreateButtonText() const
 
 FReply SCineAssemblyConfigWindow::OnCreateAssetClicked()
 {
-	UCineAssemblyFactory::CreateConfiguredAssembly(CineAssemblyToConfigure.Get(), CreateAssetPath);
+	if (ValidateAssetName())
+	{
+		UCineAssemblyFactory::CreateConfiguredAssembly(CineAssemblyToConfigure.Get(), CreateAssetPath);
 
-	RequestDestroyWindow();
+		RequestDestroyWindow();
+	}
 	return FReply::Handled();
 }
 
diff --git a/Engine/Plugins/VirtualProduction/CinematicAssemblyTools/Source/CineAssemblyToolsEditor/Private/UI/CineAssembly/SCineAssemblyConfigWindow.h b/Engine/Plugins/VirtualProduction/CinematicAssemblyTools/Source/CineAssemblyToolsEditor/Private/UI/CineAssembly/SCineAssemblyConfigWindow.h
--- a/Engine/Plugins/VirtualProduction/CinematicAssemblyTools/Source/CineAssemblyToolsEditor/Private/UI/CineAssembly/SCineAssemblyConfigWindow.h
+++ b/Engine/Plugins/VirtualProduction/CinematicAssemblyTools/Source/CineAssemblyToolsEditor/Private/UI/CineAssembly/SCineAssemblyConfigWindow.h
@@ -42,6 +42,10 @@ private:
 	/** Updates the UI and CineAssembly properties based on the selected schema */
 	void OnSchemaSelected(const FAssetData& AssetData);
 
+	/** Validate the asset name is acceptable. */
+	bool ValidateAssetName(FText& OutErrorMessage) const;
+	bool ValidateAssetName() const;
+
 private:
 	/** Transient object used only by this UI to configure the properties of the new asset that will get created by the Factory */
 	TStrongObjectPtr<UCineAssembly> CineAssemblyToConfigure = nullptr;
