CineAssembly: Add Productions level toolbar to change the active production.

diff --git a/Engine/Plugins/VirtualProduction/CinematicAssemblyTools/Source/CineAssemblyToolsEditor/Private/CineAssemblyToolsEditorModule.cpp b/Engine/Plugins/VirtualProduction/CinematicAssemblyTools/Source/CineAssemblyToolsEditor/Private/CineAssemblyToolsEditorModule.cpp
--- a/Engine/Plugins/VirtualProduction/CinematicAssemblyTools/Source/CineAssemblyToolsEditor/Private/CineAssemblyToolsEditorModule.cpp
+++ b/Engine/Plugins/VirtualProduction/CinematicAssemblyTools/Source/CineAssemblyToolsEditor/Private/CineAssemblyToolsEditorModule.cpp
@@ -16,6 +16,7 @@
 #include "Editor.h"
 #include "Framework/Application/SlateApplication.h"
 #include "Framework/Docking/TabManager.h"
+#include "LevelEditor.h"
 #include "Misc/ConfigCacheIni.h"
 #include "MoviePipelineQueue.h"
 #include "NamingTokensEngineSubsystem.h"
@@ -24,10 +25,11 @@
 #include "PropertyEditorModule.h"
 #include "Subsystems/AssetEditorSubsystem.h"
 #include "ToolMenus.h"
-#include "UI/SProductionWizard.h"
 #include "UI/CineAssembly/SCineAssemblyPropertyEditor.h"
 #include "UI/CineAssembly/SCineAssemblySchemaWindow.h"
 #include "UI/CineAssembly/SDuplicateAssemblyWindow.h"
+#include "UI/SProductionWizard.h"
+#include "Widgets/Input/SComboButton.h"
 #include "WorkspaceMenuStructure.h"
 #include "WorkspaceMenuStructureModule.h"
 
@@ -141,6 +143,8 @@ void FCineAssemblyToolsEditorModule::OnPostEngineInit()
 	AssetRegistryModule.Get().OnInMemoryAssetDeleted().AddRaw(this, &FCineAssemblyToolsEditorModule::OnAssetDeleted);
 
 	RegisterTokens();
+
+	ExtendLevelEditorToolbar();
 }
 
 void FCineAssemblyToolsEditorModule::ShutdownModule()
@@ -148,6 +152,11 @@ void FCineAssemblyToolsEditorModule::ShutdownModule()
 	FCoreDelegates::OnPostEngineInit.RemoveAll(this);
 	TakeRecorderIntegration.Reset();
 
+	if (UObjectInitialized())
+	{
+		UToolMenus::Get()->RemoveSection("LevelEditor.LevelEditorToolBar.User", "Productions");
+	}
+	
 	FPropertyEditorModule& PropertyModule = FModuleManager::GetModuleChecked<FPropertyEditorModule>("PropertyEditor");
 	PropertyModule.UnregisterCustomClassLayout(UProductionSettings::StaticClass()->GetFName());
 	PropertyModule.UnregisterCustomClassLayout(UCineAssembly::StaticClass()->GetFName());
@@ -466,6 +475,156 @@ void FCineAssemblyToolsEditorModule::SaveOpenTabs(const TMap<FTabId, FGuid>& Tab
 	}
 }
 
+void FCineAssemblyToolsEditorModule::ExtendLevelEditorToolbar()
+{
+	if (!GetDefault<UProductionSettings>()->IsLevelEditorToolbarEnabled())
+	{
+		return;
+	}
+	
+	UToolMenu* Menu = UToolMenus::Get()->ExtendMenu("LevelEditor.LevelEditorToolBar.User");
+
+	// Custom combo button so we resemble the mode selection combo.
+	const FToolMenuEntry ProductionsEntry = FToolMenuEntry::InitWidget(
+		"ProductionsMenu",
+		SNew(SComboButton)
+		.ToolTipText(LOCTEXT("LevelEditorToolbarProductionsButtonTooltip", "Change the active production"))
+		.OnGetMenuContent(FOnGetContent::CreateRaw(this, &FCineAssemblyToolsEditorModule::GenerateProductionLevelEditorToolbarMenu))
+		.ButtonContent()
+		[
+			SNew(SHorizontalBox)
+			+ SHorizontalBox::Slot()
+			.AutoWidth()
+			.VAlign(VAlign_Center)
+			.Padding(0.f, 0.f, 6.f, 0.f)
+			[
+				SNew(SBox)
+				.WidthOverride(16.f)
+				.HeightOverride(16.f)
+				[
+					SNew(SImage)
+					.Image(FSlateIcon(FAppStyle::GetAppStyleSetName(), "ProjectSettings.TabIcon").GetIcon())
+				]
+			]
+			+ SHorizontalBox::Slot()
+			.AutoWidth()
+			.VAlign(VAlign_Center)
+			.MaxWidth(200.f)
+			[
+				SNew(STextBlock)
+				.Text(
+					TAttribute<FText>::Create(TAttribute<FText>::FGetter::CreateLambda([this]()
+					{
+						const UProductionSettings* const ProductionSettings = GetDefault<UProductionSettings>();
+						TOptional<const FCinematicProduction> ActiveProduction = ProductionSettings->GetActiveProduction();
+						return ActiveProduction.IsSet()
+							? FText::FromString(ActiveProduction.GetValue().ProductionName)
+							: LOCTEXT("LevelEditorToolbarProductionsButtonLabel", "Productions");
+					}))
+				)
+				.ColorAndOpacity(FSlateColor::UseForeground())
+			]
+		],
+		FText::GetEmpty()
+	);
+	
+	Menu->FindOrAddSection("Productions").AddEntry(ProductionsEntry);
+}
+
+TSharedRef<SWidget> FCineAssemblyToolsEditorModule::GenerateProductionLevelEditorToolbarMenu()
+{
+    	const FLevelEditorModule& LevelEditorModule = FModuleManager::LoadModuleChecked<FLevelEditorModule>("LevelEditor");
+    	const TSharedPtr<FUICommandList> CommandBindings = LevelEditorModule.GetGlobalLevelEditorActions();
+    
+    	constexpr bool bShouldCloseMenuAfterSelection = false;
+    	FMenuBuilder MenuBuilder(bShouldCloseMenuAfterSelection, CommandBindings);
+	
+    	static const FName NoExtensionHook = NAME_None;
+    	
+    	{
+    		MenuBuilder.BeginSection(NoExtensionHook);
+    
+    		FUIAction OpenProductionWizardAction;
+    		OpenProductionWizardAction.ExecuteAction = FExecuteAction::CreateLambda([this]()
+    		{
+    			FGlobalTabmanager::Get()->TryInvokeTab(ProductionWizardTabName);
+    		});
+    		
+    		MenuBuilder.AddMenuEntry(LOCTEXT("OpenProductionWizard", "Production Wizard"),
+				LOCTEXT("OpenProductionsWizardToolTip", "Open the Productions Wizard"),
+				FSlateIcon(FAppStyle::GetAppStyleSetName(), "ProjectSettings.TabIcon"),
+				OpenProductionWizardAction);
+    
+    		MenuBuilder.EndSection();
+    	}
+
+		// All production options
+    	{
+    		MenuBuilder.BeginSection(FName("Productions"), LOCTEXT("ProductionsSection", "Productions"));
+    		
+    		UProductionSettings* ProductionSettings = GetMutableDefault<UProductionSettings>();
+    		const TArray<FCinematicProduction> Productions = ProductionSettings->GetProductions();
+
+    		ProductionList.Reset(Productions.Num());
+
+    		// Add empty entry
+    		{
+    			FUIAction NoProductionAction(
+				FExecuteAction::CreateLambda([this, ProductionSettings]()
+				{
+					return ProductionSettings->SetActiveProduction(FGuid());
+				}),
+				FCanExecuteAction(),
+				FIsActionChecked::CreateLambda([this, ProductionSettings]()
+				{
+					return !ProductionSettings->GetActiveProduction().IsSet();
+				}));
+
+    			MenuBuilder.AddMenuEntry(
+				LOCTEXT("NoProductionName", "No Production"),
+				LOCTEXT("NoProductionTooltip", "Clear the active production"),
+				FSlateIcon(),
+				NoProductionAction,
+				NAME_None,
+				EUserInterfaceActionType::RadioButton);
+    		}
+
+    		// Add production list
+    		for (const FCinematicProduction& Production : Productions)
+    		{
+			    TSharedRef<FCinematicProduction> ProductionRef = MakeShared<FCinematicProduction>(Production);
+    			ProductionList.Add(ProductionRef);
+
+    			FUIAction Action(
+					FExecuteAction::CreateLambda([this, ProductionRef, ProductionSettings]()
+					{
+						return ProductionSettings->SetActiveProduction(ProductionRef->ProductionID);
+					}),
+					FCanExecuteAction(),
+					FIsActionChecked::CreateLambda([this, ProductionRef, ProductionSettings]()
+					{
+						return ProductionSettings->IsActiveProduction(ProductionRef->ProductionID);
+					})
+				);
+
+    			const FText ProductionName = FText::FromString(Production.ProductionName);
+    			
+    			MenuBuilder.AddMenuEntry(
+					ProductionName,
+					ProductionName,
+					FSlateIcon(),
+					Action,
+					NAME_None,
+					EUserInterfaceActionType::RadioButton
+				);
+    		}
+    		
+    		MenuBuilder.EndSection();
+    	}
+    	
+    	return MenuBuilder.MakeWidget();
+}
+
 void FCineAssemblyToolsEditorModule::OnAssetDeleted(UObject* Object)
 {
 	// We only care about CineAssembly and CineAssemblySchema assets
diff --git a/Engine/Plugins/VirtualProduction/CinematicAssemblyTools/Source/CineAssemblyToolsEditor/Private/CineAssemblyToolsEditorModule.h b/Engine/Plugins/VirtualProduction/CinematicAssemblyTools/Source/CineAssemblyToolsEditor/Private/CineAssemblyToolsEditorModule.h
--- a/Engine/Plugins/VirtualProduction/CinematicAssemblyTools/Source/CineAssemblyToolsEditor/Private/CineAssemblyToolsEditorModule.h
+++ b/Engine/Plugins/VirtualProduction/CinematicAssemblyTools/Source/CineAssemblyToolsEditor/Private/CineAssemblyToolsEditorModule.h
@@ -13,6 +13,7 @@ class SCineAssemblyPropertyEditor;
 class SCineAssemblySchemaWindow;
 class UMoviePipelineExecutorJob;
 class ULevelSequence;
+struct FCinematicProduction;
 
 class FCineAssemblyToolsEditorModule : public IModuleInterface
 {
@@ -66,6 +67,12 @@ private:
 	/** Write out all of the tabs in the input tab map that are currently open (i.e. have a valid asset ID) */
 	void SaveOpenTabs(const TMap<FTabId, FGuid>& TabMap);
 
+	/** Extend the level editor toolbar with our menu. */
+	void ExtendLevelEditorToolbar();
+
+	/** Generates the toolbar button widget. */
+	TSharedRef<SWidget> GenerateProductionLevelEditorToolbarMenu();
+	
 private:
 	/** Manages all integration of Cinematic Assembly Tools with Take Recorder. */
 	TUniquePtr<FCineAssemblyTakeRecorderIntegration> TakeRecorderIntegration;
@@ -74,6 +81,9 @@ private:
 	TMap<FTabId, FGuid> ManagedAssemblyTabs;
 	TMap<FTabId, FGuid> ManagedSchemaTabs;
 
+	/** Available productions. */
+	TArray<TSharedPtr<FCinematicProduction>> ProductionList;
+
 	static const FString OpenTabSection;
 
 	FDelegateHandle OnSequenceSetHandle;
diff --git a/Engine/Plugins/VirtualProduction/CinematicAssemblyTools/Source/CineAssemblyToolsEditor/Private/ProductionSettings.h b/Engine/Plugins/VirtualProduction/CinematicAssemblyTools/Source/CineAssemblyToolsEditor/Private/ProductionSettings.h
--- a/Engine/Plugins/VirtualProduction/CinematicAssemblyTools/Source/CineAssemblyToolsEditor/Private/ProductionSettings.h
+++ b/Engine/Plugins/VirtualProduction/CinematicAssemblyTools/Source/CineAssemblyToolsEditor/Private/ProductionSettings.h
@@ -189,6 +189,9 @@ public:
 	FString GetUniqueProductionName() const;
 	FString GetUniqueProductionName(const FString& BaseName) const;
 
+	/** Checks if the level editor toolbar is enabled. */
+	bool IsLevelEditorToolbarEnabled() const { return bEnableProductionsLevelEditorToolbar; }
+	
 	/** Returns the delegate that broadcasts when a production is added/removed */
 	FOnProductionListChanged& OnProductionListChanged() { return ProductionListChangedDelegate; }
 
@@ -235,6 +238,10 @@ private:
 	UPROPERTY(config, EditAnywhere, Category = "Default")
 	TArray<FCinematicProduction> Productions;
 
+	/** Enable the production level editor toolbar. */
+	UPROPERTY(config, EditAnywhere, Category = "Default", meta = (ConfigRestartRequired = true, DisplayName = "Enable Production Toolbar"))
+	bool bEnableProductionsLevelEditorToolbar = true;
+
 	/** ID of the active production (in the Productions array) */
 	FGuid ActiveProductionID;
 
