ShotManagement: Adding a Sequence or Assembly as a template for Schemas

diff --git a/Engine/Plugins/VirtualProduction/CinematicAssemblyTools/Source/CineAssemblyTools/CineAssemblyTools.build.cs b/Engine/Plugins/VirtualProduction/CinematicAssemblyTools/Source/CineAssemblyTools/CineAssemblyTools.build.cs
--- a/Engine/Plugins/VirtualProduction/CinematicAssemblyTools/Source/CineAssemblyTools/CineAssemblyTools.build.cs
+++ b/Engine/Plugins/VirtualProduction/CinematicAssemblyTools/Source/CineAssemblyTools/CineAssemblyTools.build.cs
@@ -23,6 +23,7 @@ namespace UnrealBuildTool.Rules
 					"Json",
 					"MovieScene",
 					"SlateCore",
+					"UniversalObjectLocator"
 				}
 			);
 		}
diff --git a/Engine/Plugins/VirtualProduction/CinematicAssemblyTools/Source/CineAssemblyTools/Private/CineAssembly.cpp b/Engine/Plugins/VirtualProduction/CinematicAssemblyTools/Source/CineAssemblyTools/Private/CineAssembly.cpp
--- a/Engine/Plugins/VirtualProduction/CinematicAssemblyTools/Source/CineAssemblyTools/Private/CineAssembly.cpp
+++ b/Engine/Plugins/VirtualProduction/CinematicAssemblyTools/Source/CineAssemblyTools/Private/CineAssembly.cpp
@@ -13,6 +13,7 @@
 #include "Serialization/JsonSerializer.h"
 #include "Serialization/JsonWriter.h"
 #include "Tracks/MovieSceneSubTrack.h"
+#include "UniversalObjectLocator.h"
 #include "UObject/AssetRegistryTagsContext.h"
 
 #if WITH_EDITOR
@@ -120,6 +121,47 @@ void UCineAssembly::Initialize()
 	Super::Initialize();
 }
 
+void UCineAssembly::InitializeFromTemplate(ULevelSequence* Template)
+{
+	Initialize();
+
+	if (Template == nullptr)
+	{
+		return;
+	}
+
+	if (UMovieScene* TemplateMovieScene = Template->GetMovieScene())
+	{
+		// Attempt to duplicate the inner MovieScene. Assign only if duplication is a success to avoid leaving the assembly in a partially valid state.
+		// Duplicating the MovieScene will null any event track bindings which is the currently supported behaviour.
+		UMovieScene* DuplicateMovieScene = DuplicateObject(TemplateMovieScene, this);
+
+		if (DuplicateMovieScene == nullptr)
+		{
+			return;
+		}
+
+		MovieScene = DuplicateMovieScene;
+
+		// Fix up supported bindings
+		if (const FMovieSceneBindingReferences* TemplateBindingReferences = Template->GetBindingReferences())
+		{
+			FMovieSceneBindingReferences& ThisBindingReferencesRef = BindingReferences;
+
+			for (const FMovieSceneBindingReference& Reference : TemplateBindingReferences->GetAllReferences())
+			{
+				FUniversalObjectLocator NewLocator = Reference.Locator;
+
+				UMovieSceneCustomBinding* CustomBindingDuplicate = Reference.CustomBinding != nullptr 
+					? Cast<UMovieSceneCustomBinding>(StaticDuplicateObject(Reference.CustomBinding, MovieScene))
+					: nullptr;
+				
+				ThisBindingReferencesRef.AddBinding(Reference.ID, MoveTemp(NewLocator), Reference.ResolveFlags, CustomBindingDuplicate);
+			}
+		}
+	}
+}
+
 FGuid UCineAssembly::GetAssemblyGuid() const
 {
 	return AssemblyGuid;
@@ -295,7 +337,21 @@ void UCineAssembly::CreateSubAssemblies()
 		UCineAssembly* SubAssembly = Cast<UCineAssembly>(AssetTools.CreateAsset(UniqueAssetName, SubAssemblyPath, UCineAssembly::StaticClass(), nullptr));
 		if (SubAssembly)
 		{
-			SubAssembly->Initialize();
+			ULevelSequence* Template = nullptr;
+			// Check to see if there is a supplied templated for this Subsequence.
+			if (FSoftObjectPath* TemplatePath = BaseSchema->SubsequenceTemplates.Find(SubAssemblyName.Template))
+			{
+				Template = Cast<ULevelSequence>(TemplatePath->TryLoad());
+			}
+
+			if (Template != nullptr)
+			{
+				SubAssembly->InitializeFromTemplate(Template);
+			}
+			else
+			{
+				SubAssembly->Initialize();
+			}
 
 			const TRange<FFrameNumber> PlaybackRange = MovieScene->GetPlaybackRange();
 			SubAssembly->GetMovieScene()->SetPlaybackRange(PlaybackRange);
diff --git a/Engine/Plugins/VirtualProduction/CinematicAssemblyTools/Source/CineAssemblyTools/Public/CineAssembly.h b/Engine/Plugins/VirtualProduction/CinematicAssemblyTools/Source/CineAssemblyTools/Public/CineAssembly.h
--- a/Engine/Plugins/VirtualProduction/CinematicAssemblyTools/Source/CineAssemblyTools/Public/CineAssembly.h
+++ b/Engine/Plugins/VirtualProduction/CinematicAssemblyTools/Source/CineAssemblyTools/Public/CineAssembly.h
@@ -29,6 +29,9 @@ public:
 	CINEASSEMBLYTOOLS_API virtual void Initialize() override;
 	//~ End ULevelSquence Interface
 
+	/** Given a ULevelSequence as a template, initialize this assembly from it. */
+	CINEASSEMBLYTOOLS_API void InitializeFromTemplate(ULevelSequence* Template);
+
 	//~ Begin UObject Interface
 	virtual void GetAssetRegistryTags(FAssetRegistryTagsContext Context) const override;
 	virtual void PostInitProperties() override;
diff --git a/Engine/Plugins/VirtualProduction/CinematicAssemblyTools/Source/CineAssemblyTools/Public/CineAssemblySchema.h b/Engine/Plugins/VirtualProduction/CinematicAssemblyTools/Source/CineAssemblyTools/Public/CineAssemblySchema.h
--- a/Engine/Plugins/VirtualProduction/CinematicAssemblyTools/Source/CineAssemblyTools/Public/CineAssemblySchema.h
+++ b/Engine/Plugins/VirtualProduction/CinematicAssemblyTools/Source/CineAssemblyTools/Public/CineAssemblySchema.h
@@ -129,6 +129,14 @@ public:
 	UPROPERTY()
 	TArray<FString> FoldersToCreate;
 
+	/** Template to use when creating the assembly as a starting point. */
+	UPROPERTY()
+	FSoftObjectPath Template;
+
+	/** Templates to use for subassemblies that are created as part of this assembly. */
+	UPROPERTY()
+	TMap<FString, FSoftObjectPath> SubsequenceTemplates;
+
 	/** Mark this Schema as Data-Only. Assemblies made from this Schema will have their Data-Only defaulted to this value. */
 	UPROPERTY(AssetRegistrySearchable, BlueprintReadWrite, EditAnywhere, Category = "Default", AdvancedDisplay, meta = (DisplayName = "Data-Only"))
 	bool bIsDataOnly = false;
diff --git a/Engine/Plugins/VirtualProduction/CinematicAssemblyTools/Source/CineAssemblyToolsEditor/Private/CineAssemblyFactory.cpp b/Engine/Plugins/VirtualProduction/CinematicAssemblyTools/Source/CineAssemblyToolsEditor/Private/CineAssemblyFactory.cpp
--- a/Engine/Plugins/VirtualProduction/CinematicAssemblyTools/Source/CineAssemblyToolsEditor/Private/CineAssemblyFactory.cpp
+++ b/Engine/Plugins/VirtualProduction/CinematicAssemblyTools/Source/CineAssemblyToolsEditor/Private/CineAssemblyFactory.cpp
@@ -143,17 +143,32 @@ void UCineAssemblyFactory::CreateConfiguredAssembly(UCineAssembly* ConfiguredAss
 	// If the configured assembly already has a MovieScene object (which could happen in case of a duplication), do not overwrite it with a new one
 	if (!ConfiguredAssembly->GetMovieScene())
 	{
-		// Do the same setup that ULevelSequence assets do when they are created
-		ConfiguredAssembly->Initialize();
-
-		if (UMovieScene* MovieScene = ConfiguredAssembly->GetMovieScene())
+		// if we have a template to build from, setup now
+		bool bBuiltFromTemplate = false;
+		if (const UCineAssemblySchema* Schema = ConfiguredAssembly->GetSchema())
 		{
-			const UMovieSceneToolsProjectSettings* MovieSceneToolsProjectSettings = GetDefault<UMovieSceneToolsProjectSettings>();
+			if (ULevelSequence* Template = Cast<ULevelSequence>(Schema->Template.TryLoad()))
+			{
+				ConfiguredAssembly->InitializeFromTemplate(Template);
+				// Confirm we have a valid MovieScene built.
+				bBuiltFromTemplate = ConfiguredAssembly->GetMovieScene() != nullptr;
+			}
+		}
 
-			const FFrameRate TickResolution = MovieScene->GetTickResolution();
-			const FFrameNumber DefaultStartFrame = (MovieSceneToolsProjectSettings->DefaultStartTime * TickResolution).FloorToFrame();
-			const int32 DefaultDuration = (MovieSceneToolsProjectSettings->DefaultDuration * TickResolution).FloorToFrame().Value;
-			MovieScene->SetPlaybackRange(DefaultStartFrame, DefaultDuration);
+		if (!bBuiltFromTemplate)
+		{
+			// Do the same setup that ULevelSequence assets do when they are created
+			ConfiguredAssembly->Initialize();
+
+			if (UMovieScene* MovieScene = ConfiguredAssembly->GetMovieScene())
+			{
+				const UMovieSceneToolsProjectSettings* MovieSceneToolsProjectSettings = GetDefault<UMovieSceneToolsProjectSettings>();
+
+				const FFrameRate TickResolution = MovieScene->GetTickResolution();
+				const FFrameNumber DefaultStartFrame = (MovieSceneToolsProjectSettings->DefaultStartTime * TickResolution).FloorToFrame();
+				const int32 DefaultDuration = (MovieSceneToolsProjectSettings->DefaultDuration * TickResolution).FloorToFrame().Value;
+				MovieScene->SetPlaybackRange(DefaultStartFrame, DefaultDuration);
+			}
 		}
 	}
 
diff --git a/Engine/Plugins/VirtualProduction/CinematicAssemblyTools/Source/CineAssemblyToolsEditor/Private/UI/CineAssembly/SCineAssemblySchemaWindow.cpp b/Engine/Plugins/VirtualProduction/CinematicAssemblyTools/Source/CineAssemblyToolsEditor/Private/UI/CineAssembly/SCineAssemblySchemaWindow.cpp
--- a/Engine/Plugins/VirtualProduction/CinematicAssemblyTools/Source/CineAssemblyToolsEditor/Private/UI/CineAssembly/SCineAssemblySchemaWindow.cpp
+++ b/Engine/Plugins/VirtualProduction/CinematicAssemblyTools/Source/CineAssemblyToolsEditor/Private/UI/CineAssembly/SCineAssemblySchemaWindow.cpp
@@ -585,7 +585,7 @@ TSharedRef<SWidget> SCineAssemblySchemaWindow::MakeHierarchyTabContent()
 			SNew(SBorder)
 				.BorderImage(FCineAssemblyToolsStyle::Get().GetBrush("ProductionWizard.RecessedBackground"))
 				[
-					SAssignNew(AssetListView, SListView<TSharedPtr<FString>>)
+					SAssignNew(AssetListView, SListView<TSharedPtr<FSchemaAssetParameters>>)
 						.ListItemsSource(&AssetListItems)
 						.SelectionMode(ESelectionMode::Single)
 						.OnGenerateRow(this, &SCineAssemblySchemaWindow::OnGenerateAssetRow)
@@ -654,7 +654,7 @@ TSharedRef<SWidget> SCineAssemblySchemaWindow::MakeHierarchyTabContent()
 FReply SCineAssemblySchemaWindow::OnAddAssetClicked()
 {
 	// Add a new empty string to the list, which will be renamed by the user
-	AssetListItems.Add(MakeShared<FString>(TEXT("")));
+	AssetListItems.Add(MakeShared<FSchemaAssetParameters>(FSchemaAssetParameters()));
 	AssetListView->RequestListRefresh();
 	return FReply::Handled();
 }
@@ -728,25 +728,58 @@ FReply SCineAssemblySchemaWindow::OnAddFolderClicked()
 	return FReply::Handled();
 }
 
-TSharedRef<ITableRow> SCineAssemblySchemaWindow::OnGenerateAssetRow(TSharedPtr<FString> InItem, const TSharedRef<STableViewBase>& OwnerTable)
+TSharedRef<ITableRow> SCineAssemblySchemaWindow::OnGenerateAssetRow(TSharedPtr<FSchemaAssetParameters> InItem, const TSharedRef<STableViewBase>& OwnerTable)
 {
 	TSharedRef<SSchemaAssetTableRow> Row = SNew(SSchemaAssetTableRow, OwnerTable)
 		.ShowSelection(true)
 		.Padding(FMargin(4.0f, 2.0f))
-		.OnDragDetected(this, &SCineAssemblySchemaWindow::OnAssetRowDragDetected);
+		.OnDragDetected(this, &SCineAssemblySchemaWindow::OnAssetRowDragDetected)
+		.ToolTipText(LOCTEXT("AssetRowTooltip", "Double click to rename the asset and change the template asset."));
 
 	// Store a reference to the editable textblock in the row to easily set it to edit mode for renaming
 	Row->TextBlock = SNew(SInlineEditableTextBlock)
 		.Text_Lambda([InItem]() 
 		{ 
-			return FText::FromString(*InItem.Get()); 
+			return FText::FromString(InItem->Name); 
 		})
 		.OnVerifyTextChanged(this, &SCineAssemblySchemaWindow::ValidateAssetName)
 		.OnTextCommitted_Lambda([this, InItem](const FText& InText, ETextCommit::Type InCommitType) 
 		{ 
-			*InItem = InText.ToString();
+			InItem->Name = InText.ToString();
 		});
 
+	Row->TemplateObjectEntry = SNew(SObjectPropertyEntryBox)
+		.AllowCreate(true)
+		.AllowClear(true)
+		.DisplayThumbnail(false)
+		.DisplayCompactSize(true)
+		.AllowedClass(ULevelSequence::StaticClass())
+		.ToolTipText(LOCTEXT("SchemaAssetTemplateSelectionTooltip", "Select an existing Asset to use as a template for this asset."))
+		.Visibility_Lambda([WeakRow = TWeakPtr<SSchemaAssetTableRow>(Row)]() 
+			{ 
+				if (TSharedPtr<SSchemaAssetTableRow> RowPtr = WeakRow.Pin())
+				{
+					return RowPtr->TextBlock->IsInEditMode() ? EVisibility::Visible : EVisibility::HitTestInvisible;
+				}
+
+				return EVisibility::Hidden;
+			})
+		.ObjectPath_Lambda([WeakInItem = TWeakPtr<FSchemaAssetParameters>(InItem)]() -> FString
+			{
+				if (TSharedPtr<FSchemaAssetParameters> ItemPtr = WeakInItem.Pin())
+				{
+					return ItemPtr->Template.ToString();
+				}
+				return FString();
+			})
+		.OnObjectChanged_Lambda([WeakInItem = TWeakPtr<FSchemaAssetParameters>(InItem)](const FAssetData& InAssetData)
+			{
+				if (TSharedPtr<FSchemaAssetParameters> ItemPtr = WeakInItem.Pin())
+				{
+					ItemPtr->Template = InAssetData.GetSoftObjectPath();
+				}
+			});
+
 	Row->SetContent(SNew(SBorder)
 		.Padding(8.0f)
 		.BorderImage(FCineAssemblyToolsStyle::Get().GetBrush("ProductionWizard.PanelBackground"))
@@ -768,6 +801,13 @@ TSharedRef<ITableRow> SCineAssemblySchemaWindow::OnGenerateAssetRow(TSharedPtr<F
 					Row->TextBlock.ToSharedRef()
 				]
 
+			+ SHorizontalBox::Slot()
+				.FillWidth(1.0f)
+				.VAlign(VAlign_Center)
+				[
+					Row->TemplateObjectEntry.ToSharedRef()
+				]
+
 			+ SHorizontalBox::Slot()
 				.AutoWidth()
 				.VAlign(VAlign_Center)
@@ -784,10 +824,10 @@ void SCineAssemblySchemaWindow::OnAssetListRebuilt()
 	if (AssetListView && AssetListView->GetItems().Num() > 0)
 	{
 		// Only trigger edit mode on the last item in the asset list, and only if it is an empty string so that it can be renamed
-		TSharedPtr<FString> LastItem = AssetListView->GetItems().Last();
-		if (LastItem && LastItem->IsEmpty())
+		TSharedPtr<FSchemaAssetParameters> LastItem = AssetListView->GetItems().Last();
+		if (LastItem && LastItem->Name.IsEmpty())
 		{
-			*LastItem = TEXT("NewAsset");
+			LastItem->Name = TEXT("NewAsset");
 			EnterEditMode(LastItem);
 		}
 	}
@@ -796,13 +836,13 @@ void SCineAssemblySchemaWindow::OnAssetListRebuilt()
 TSharedPtr<SWidget> SCineAssemblySchemaWindow::MakeAssetListContextMenu()
 {
 	// The ListView uses single selection mode, so at most one item can ever be selected by the user
-	TArray<TSharedPtr<FString>> SelectedItems = AssetListView->GetSelectedItems();
+	TArray<TSharedPtr<FSchemaAssetParameters>> SelectedItems = AssetListView->GetSelectedItems();
 	if (SelectedItems.Num() == 1)
 	{
 		constexpr bool bCloseAfterMenuSelection = true;
 		FMenuBuilder MenuBuilder(bCloseAfterMenuSelection, nullptr);
 
-		TSharedPtr<FString>& SelectedItem = SelectedItems[0];
+		TSharedPtr<FSchemaAssetParameters>& SelectedItem = SelectedItems[0];
 
 		MenuBuilder.AddMenuEntry(
 			LOCTEXT("RenameAsset", "Rename"),
@@ -833,10 +873,10 @@ FReply SCineAssemblySchemaWindow::OnAssetListKeyDown(const FGeometry& MyGeometry
 	if (InKeyEvent.GetKey() == EKeys::Delete)
 	{
 		// The list view uses single selection mode, so at most one item can ever be selected by the user
-		TArray<TSharedPtr<FString>> SelectedItems = AssetListView->GetSelectedItems();
+		TArray<TSharedPtr<FSchemaAssetParameters>> SelectedItems = AssetListView->GetSelectedItems();
 		if (SelectedItems.Num() == 1)
 		{
-			const TSharedPtr<FString>& SelectedItem = SelectedItems[0];
+			const TSharedPtr<FSchemaAssetParameters>& SelectedItem = SelectedItems[0];
 			DeleteAssetItem(SelectedItem);
 		}
 	}
@@ -869,6 +909,20 @@ TSharedRef<ITableRow> SCineAssemblySchemaWindow::OnGenerateTreeRow(TSharedPtr<FS
 			})
 		.OnAcceptDrop(this, &SCineAssemblySchemaWindow::OnTreeRowAcceptDrop)
 		.OnDragDetected(this, &SCineAssemblySchemaWindow::OnTreeRowDragDetected)
+		.ToolTipText_Lambda([this, WeakTreeItem = TWeakPtr<FSchemaTreeItem>(TreeItem)]() -> FText
+			{
+				if (TSharedPtr<FSchemaTreeItem> ItemPtr = WeakTreeItem.Pin())
+				{
+					if (ItemPtr->Type == FSchemaTreeItem::EItemType::Asset)
+					{
+						return ItemPtr == TopLevelAssemblyItem
+							? LOCTEXT("AssetTreeTopLevelAssemblyTooltip", "Right click to change the Assemblies template asset.")
+							: LOCTEXT("AssetTreeSubAssemblyTooltip", "Right click to rename the SubAssembly or change the template asset.");
+					}
+				}
+
+				return FText::GetEmpty();
+			})
 		.Content()
 		[
 			SNew(SHorizontalBox)
@@ -903,6 +957,28 @@ TSharedRef<ITableRow> SCineAssemblySchemaWindow::OnGenerateTreeRow(TSharedPtr<FS
 							})
 						.OnTextCommitted(this, &SCineAssemblySchemaWindow::OnTreeItemTextCommitted, TreeItem)
 				]
+
+			+ SHorizontalBox::Slot()
+				.VAlign(VAlign_Center)
+				[
+					SAssignNew(TreeItem->TemplateObjectEntry, SObjectPropertyEntryBox)
+						.AllowCreate(true)
+						.AllowClear(true)
+						.DisplayThumbnail(false)
+						.DisplayCompactSize(true)
+						.AllowedClass(ULevelSequence::StaticClass())
+						.ToolTipText(LOCTEXT("SchemaAssetTemplateSelectionTooltip", "Select an existing Asset to use as a template for this asset."))
+						.Visibility(TreeItem->Type == FSchemaTreeItem::EItemType::Asset ? EVisibility::HitTestInvisible : EVisibility::Collapsed)
+						.ObjectPath_Lambda([WeakTreeItem = TWeakPtr<FSchemaTreeItem>(TreeItem)]() -> FString
+							{
+								if (TSharedPtr<FSchemaTreeItem> ItemPtr = WeakTreeItem.Pin())
+								{
+									return ItemPtr->AssetTemplate.ToString();
+								}
+								return FString();
+							})
+						.OnObjectChanged(this, &SCineAssemblySchemaWindow::OnTreeItemTemplateChanged, TreeItem)
+				]
 		];
 }
 
@@ -947,30 +1023,54 @@ TSharedPtr<SWidget> SCineAssemblySchemaWindow::MakeContentTreeContextMenu()
 			);
 		}
 
-		if ((SelectedTreeItem != RootItem) && (SelectedTreeItem != TopLevelAssemblyItem))
+		if (SelectedTreeItem != RootItem)
 		{
-			MenuBuilder.AddMenuEntry(
-				LOCTEXT("RenameAction", "Rename"),
-				LOCTEXT("RenameActionToolTip", "Rename"),
-				FSlateIcon(FCineAssemblyToolsStyle::StyleName, "Icons.AssetNaming"),
-				FUIAction(
-					FExecuteAction::CreateLambda([SelectedTreeItem]()
-						{
-							FSlateApplication::Get().SetKeyboardFocus(SelectedTreeItem->NameWidget.ToSharedRef());
-							SelectedTreeItem->NameWidget->EnterEditingMode();
-						})),
-				NAME_None,
-				EUserInterfaceActionType::Button
-			);
+			if (SelectedTreeItem->Type == FSchemaTreeItem::EItemType::Asset)
+			{
+				MenuBuilder.AddMenuEntry(
+					LOCTEXT("EditTemplateAction", "Edit Template"),
+					LOCTEXT("EditTemplateActionToolTip", "Edit the assets template"),
+					FSlateIcon(FAppStyle::GetAppStyleSetName(), "Icons.Edit"),
+					FUIAction(
+						FExecuteAction::CreateLambda([SelectedTreeItem]()
+							{
+								if (SelectedTreeItem->TemplateObjectEntry.IsValid())
+								{
+									FSlateApplication::Get().SetKeyboardFocus(SelectedTreeItem->TemplateObjectEntry.ToSharedRef());
+									SelectedTreeItem->TemplateObjectEntry->SetVisibility(EVisibility::Visible);
+									SelectedTreeItem->TemplateObjectEntry->OpenEntryBox();
+								}
+							})),
+					NAME_None,
+					EUserInterfaceActionType::Button
+				);
+			}
 
-			MenuBuilder.AddMenuEntry(
-				LOCTEXT("DeleteAction", "Delete"),
-				LOCTEXT("DeleteActionToolTip", "Delete"),
-				FSlateIcon(FAppStyle::GetAppStyleSetName(), "Icons.Delete"),
-				FUIAction(FExecuteAction::CreateSP(this, &SCineAssemblySchemaWindow::DeleteTreeItem, SelectedTreeItem)),
-				NAME_None,
-				EUserInterfaceActionType::Button
-			);
+			if (SelectedTreeItem != TopLevelAssemblyItem)
+			{
+				MenuBuilder.AddMenuEntry(
+					LOCTEXT("RenameAction", "Rename"),
+					LOCTEXT("RenameActionToolTip", "Rename"),
+					FSlateIcon(FCineAssemblyToolsStyle::StyleName, "Icons.AssetNaming"),
+					FUIAction(
+						FExecuteAction::CreateLambda([SelectedTreeItem]()
+							{
+								FSlateApplication::Get().SetKeyboardFocus(SelectedTreeItem->NameWidget.ToSharedRef());
+								SelectedTreeItem->NameWidget->EnterEditingMode();
+							})),
+					NAME_None,
+					EUserInterfaceActionType::Button
+				);
+
+				MenuBuilder.AddMenuEntry(
+					LOCTEXT("DeleteAction", "Delete"),
+					LOCTEXT("DeleteActionToolTip", "Delete"),
+					FSlateIcon(FAppStyle::GetAppStyleSetName(), "Icons.Delete"),
+					FUIAction(FExecuteAction::CreateSP(this, &SCineAssemblySchemaWindow::DeleteTreeItem, SelectedTreeItem)),
+					NAME_None,
+					EUserInterfaceActionType::Button
+				);
+			}
 		}
 
 		return MenuBuilder.MakeWidget();
@@ -1007,7 +1107,7 @@ void SCineAssemblySchemaWindow::InitializeContentTree()
 
 	TreeItems.Add(RootItem);
 
-	auto AddItemsToTree = [this](TArray<FString>& ItemList, FSchemaTreeItem::EItemType ItemType)
+	auto AddItemsToTree = [this](TArray<FString>& ItemList, const TMap<FString, FSoftObjectPath>& Templates, FSchemaTreeItem::EItemType ItemType)
 		{
 			// Sort the list so that paths are added to the tree in the proper order
 			ItemList.Sort();
@@ -1025,6 +1125,11 @@ void SCineAssemblySchemaWindow::InitializeContentTree()
 					NewItem->Path = ItemName;
 					NewItem->Parent = ParentItem;
 
+					if (const FSoftObjectPath* FoundTemplate = Templates.Find(ItemName))
+					{
+						NewItem->AssetTemplate = *FoundTemplate;
+					}
+
 					if (ItemType == FSchemaTreeItem::EItemType::Folder)
 					{
 						ParentItem->ChildFolders.Add(NewItem);
@@ -1037,8 +1142,8 @@ void SCineAssemblySchemaWindow::InitializeContentTree()
 			}
 		};
 
-	AddItemsToTree(SchemaToConfigure->FoldersToCreate, FSchemaTreeItem::EItemType::Folder);
-	AddItemsToTree(SchemaToConfigure->SubsequencesToCreate, FSchemaTreeItem::EItemType::Asset);
+	AddItemsToTree(SchemaToConfigure->FoldersToCreate, {}, FSchemaTreeItem::EItemType::Folder);
+	AddItemsToTree(SchemaToConfigure->SubsequencesToCreate, SchemaToConfigure->SubsequenceTemplates, FSchemaTreeItem::EItemType::Asset);
 
 	// Add the top-level assembly node
 	TSharedPtr<FSchemaTreeItem> ParentItem = FindItemAtPathRecursive(RootItem, SchemaToConfigure->DefaultAssemblyPath);
@@ -1047,6 +1152,7 @@ void SCineAssemblySchemaWindow::InitializeContentTree()
 		TopLevelAssemblyItem = MakeShared<FSchemaTreeItem>();
 		TopLevelAssemblyItem->Type = FSchemaTreeItem::EItemType::Asset;
 		TopLevelAssemblyItem->Path = TEXT("{assembly}");
+		TopLevelAssemblyItem->AssetTemplate = SchemaToConfigure->Template;
 		TopLevelAssemblyItem->Parent = ParentItem;
 
 		ParentItem->ChildAssets.Add(TopLevelAssemblyItem);
@@ -1072,19 +1178,20 @@ void SCineAssemblySchemaWindow::GetFolderListRecursive(TSharedPtr<FSchemaTreeIte
 	}
 }
 
-void SCineAssemblySchemaWindow::GetAssetListRecursive(TSharedPtr<FSchemaTreeItem> TreeItem, TArray<FString>& AssetPathList) const
+void SCineAssemblySchemaWindow::GetAssetListRecursive(TSharedPtr<FSchemaTreeItem> TreeItem, TArray<FString>& AssetPathList, TMap<FString, FSoftObjectPath>& Templates) const
 {
 	for (const TSharedPtr<FSchemaTreeItem>& Asset : TreeItem->ChildAssets)
 	{
 		if (Asset != TopLevelAssemblyItem)
 		{
 			AssetPathList.Add(Asset->Path);
+			Templates.Add(Asset->Path, Asset->AssetTemplate);
 		}
 	}
 
 	for (const TSharedPtr<FSchemaTreeItem>& Child : TreeItem->ChildFolders)
 	{
-		GetAssetListRecursive(Child, AssetPathList);
+		GetAssetListRecursive(Child, AssetPathList, Templates);
 	}
 }
 
@@ -1127,7 +1234,7 @@ void SCineAssemblySchemaWindow::SetChildrenPathRecursive(TSharedPtr<FSchemaTreeI
 	}
 }
 
-void SCineAssemblySchemaWindow::EnterEditMode(TSharedPtr<FString> ItemToRename)
+void SCineAssemblySchemaWindow::EnterEditMode(TSharedPtr<FSchemaAssetParameters> ItemToRename)
 {
 	// Get the row for the input item and put its textblock in edit mode so the user can rename the item
 	if (TSharedPtr<SSchemaAssetTableRow> Widget = StaticCastSharedPtr<SSchemaAssetTableRow>(AssetListView->WidgetFromItem(ItemToRename)))
@@ -1171,6 +1278,28 @@ void SCineAssemblySchemaWindow::OnTreeItemTextCommitted(const FText& InText, ETe
 	TreeView->RequestTreeRefresh();
 }
 
+void SCineAssemblySchemaWindow::OnTreeItemTemplateChanged(const FAssetData& InAssetData, TSharedPtr<FSchemaTreeItem> TreeItem)
+{
+	// Early-out if this is not an asset item
+	if (TreeItem->Type != FSchemaTreeItem::EItemType::Asset)
+	{
+		return;
+	}
+
+	const FSoftObjectPath OldSoftPath = TreeItem->AssetTemplate;
+	const FSoftObjectPath NewSoftPath = InAssetData.GetSoftObjectPath();
+	// Early-out if the name has not actually changed
+	if (OldSoftPath == NewSoftPath)
+	{
+		return;
+	}
+
+	TreeItem->AssetTemplate = NewSoftPath;
+	TreeItem->TemplateObjectEntry->SetVisibility(EVisibility::HitTestInvisible);
+	UpdateAssetList();
+	TreeView->RequestTreeRefresh();
+}
+
 bool SCineAssemblySchemaWindow::ValidateAssetName(const FText& InText, FText& OutErrorMessage) const
 {
 	// An empty name is invalid
@@ -1222,7 +1351,7 @@ bool SCineAssemblySchemaWindow::ValidateFolderName(const FText& InText, FText& O
 	return true;
 }
 
-void SCineAssemblySchemaWindow::DeleteAssetItem(TSharedPtr<FString> InItem)
+void SCineAssemblySchemaWindow::DeleteAssetItem(TSharedPtr<FSchemaAssetParameters> InItem)
 {
 	const int32 IndexToRemove = AssetListItems.IndexOfByKey(InItem);
 	if (AssetListItems.IsValidIndex(IndexToRemove))
@@ -1338,11 +1467,13 @@ void SCineAssemblySchemaWindow::UpdateFolderList()
 
 void SCineAssemblySchemaWindow::UpdateAssetList()
 {
-	// Update the schema's asset list
+	// Update the schema's asset list and templates
 	const TArray<FString> CachedAssetList = SchemaToConfigure->SubsequencesToCreate;
+	const TMap<FString, FSoftObjectPath> CachedSubsequenceTemplates = SchemaToConfigure->SubsequenceTemplates;
 
 	TArray<FString> AssetList;
-	GetAssetListRecursive(RootItem, AssetList);
+	TMap<FString, FSoftObjectPath> Templates;
+	GetAssetListRecursive(RootItem, AssetList, Templates);
 
 	SchemaToConfigure->DefaultAssemblyPath = TopLevelAssemblyItem->Parent->Path;
 
@@ -1351,6 +1482,33 @@ void SCineAssemblySchemaWindow::UpdateAssetList()
 		SchemaToConfigure->Modify();
 		SchemaToConfigure->SubsequencesToCreate = AssetList;
 	}
+
+	if (SchemaToConfigure->Template != TopLevelAssemblyItem->AssetTemplate)
+	{
+		SchemaToConfigure->Modify();
+		SchemaToConfigure->Template = TopLevelAssemblyItem->AssetTemplate;
+	}
+
+	bool bSubsequenceTemplatesHaveChanged = Templates.Num() != CachedSubsequenceTemplates.Num();
+
+	if (!bSubsequenceTemplatesHaveChanged)
+	{
+		for (auto Iter = Templates.CreateConstIterator(); Iter; ++Iter)
+		{
+			if (const FSoftObjectPath* Found = CachedSubsequenceTemplates.Find(Iter->Key);
+				Found == nullptr || (Found != nullptr && *Found != Iter->Value))
+			{
+				bSubsequenceTemplatesHaveChanged = true;
+				break;
+			}
+		}
+	}
+
+	if (bSubsequenceTemplatesHaveChanged)
+	{
+		SchemaToConfigure->Modify();
+		SchemaToConfigure->SubsequenceTemplates = Templates;
+	}
 }
 
 FReply SCineAssemblySchemaWindow::OnAssetRowDragDetected(const FGeometry& MyGeometry, const FPointerEvent& MouseEvent)
@@ -1358,16 +1516,21 @@ FReply SCineAssemblySchemaWindow::OnAssetRowDragDetected(const FGeometry& MyGeom
 	if (MouseEvent.IsMouseButtonDown(EKeys::LeftMouseButton))
 	{
 		// The ListView uses single selection mode, so at most one item can ever be selected by the user
-		TArray<TSharedPtr<FString>> SelectedItems = AssetListView->GetSelectedItems();
+		TArray<TSharedPtr<FSchemaAssetParameters>> SelectedItems = AssetListView->GetSelectedItems();
 		if (SelectedItems.Num() == 1)
 		{
-			const TSharedPtr<FString>& SelectedItem = SelectedItems[0];
+			const TSharedPtr<FSchemaAssetParameters>& SelectedItem = SelectedItems[0];
 			TSharedRef<FSchemaAssetDragDrop> Operation = FSchemaAssetDragDrop::New(*SelectedItem);
 
-			Operation->OnDropNotHandled.BindLambda([this](const FString& AssetName)
+			Operation->OnDropNotHandled.BindLambda([this, WeakSelectedItem = TWeakPtr<FSchemaAssetParameters>(SelectedItem)](const FString& AssetName)
 				{
+					FSoftObjectPath RevertTemplate;
+					if (TSharedPtr<FSchemaAssetParameters> SelectedItemPtr = WeakSelectedItem.Pin())
+					{
+						RevertTemplate = SelectedItemPtr->Template;
+					}
 					// If the asset item is not dropped in a valid place, restore it to the asset list
-					AssetListItems.Add(MakeShared<FString>(AssetName));
+					AssetListItems.Add(MakeShared<FSchemaAssetParameters>(AssetName, RevertTemplate));
 					AssetListView->RequestListRefresh();
 				});
 
@@ -1396,7 +1559,7 @@ FReply SCineAssemblySchemaWindow::OnTreeRowDragDetected(const FGeometry& MyGeome
 				return FReply::Handled();
 			}
 			
-			TSharedRef<FSchemaAssetDragDrop> Operation = FSchemaAssetDragDrop::New(FPaths::GetPathLeaf(SelectedTreeItem->Path));
+			TSharedRef<FSchemaAssetDragDrop> Operation = FSchemaAssetDragDrop::New(FSchemaAssetParameters(FPaths::GetPathLeaf(SelectedTreeItem->Path), SelectedTreeItem->AssetTemplate));
 			Operation->SourceTreeItem = SelectedTreeItem;
 
 			if (SelectedTreeItem->Type == FSchemaTreeItem::EItemType::Asset)
@@ -1467,6 +1630,7 @@ FReply SCineAssemblySchemaWindow::OnTreeRowAcceptDrop(const FDragDropEvent& InDr
 
 		const FString ParentPath = InItem->Path;
 		NewItem->Path = ParentPath / Operation->ItemName;
+		NewItem->AssetTemplate = Operation->Template;
 
 		NewItem->Parent = InItem;
 
@@ -1493,10 +1657,11 @@ FReply SCineAssemblySchemaWindow::OnTreeRowAcceptDrop(const FDragDropEvent& InDr
 	return FReply::Handled();
 }
 
-TSharedRef<FSchemaAssetDragDrop> FSchemaAssetDragDrop::New(const FString& InName)
+TSharedRef<FSchemaAssetDragDrop> FSchemaAssetDragDrop::New(const FSchemaAssetParameters& InParameters)
 {
 	TSharedRef<FSchemaAssetDragDrop> DragDropOp = MakeShared<FSchemaAssetDragDrop>();
-	DragDropOp->ItemName = InName;
+	DragDropOp->ItemName = InParameters.Name;
+	DragDropOp->Template = InParameters.Template;
 	DragDropOp->MouseCursor = EMouseCursor::GrabHandClosed;
 	DragDropOp->Construct();
 
diff --git a/Engine/Plugins/VirtualProduction/CinematicAssemblyTools/Source/CineAssemblyToolsEditor/Private/UI/CineAssembly/SCineAssemblySchemaWindow.h b/Engine/Plugins/VirtualProduction/CinematicAssemblyTools/Source/CineAssemblyToolsEditor/Private/UI/CineAssembly/SCineAssemblySchemaWindow.h
--- a/Engine/Plugins/VirtualProduction/CinematicAssemblyTools/Source/CineAssemblyToolsEditor/Private/UI/CineAssembly/SCineAssemblySchemaWindow.h
+++ b/Engine/Plugins/VirtualProduction/CinematicAssemblyTools/Source/CineAssemblyToolsEditor/Private/UI/CineAssembly/SCineAssemblySchemaWindow.h
@@ -6,6 +6,7 @@
 
 #include "CineAssemblySchema.h"
 #include "Input/DragAndDrop.h"
+#include "PropertyCustomizationHelpers.h"
 #include "Widgets/Layout/SWidgetSwitcher.h"
 #include "Widgets/Text/SInlineEditableTextBlock.h"
 #include "Widgets/Views/SListView.h"
@@ -20,12 +21,25 @@ enum class ESchemaConfigMode : uint8
 	Edit
 };
 
+/** Schema asset parameters describing a configured asset for the schema. */
+struct FSchemaAssetParameters
+{
+	/** Name of the asset. */
+	FString Name;
+
+	/** Template object to use to create the asset. */
+	FSoftObjectPath Template;
+};
+
 /** Row widget for the asset list view */
-class SSchemaAssetTableRow : public STableRow<TSharedPtr<FString>>
+class SSchemaAssetTableRow : public STableRow<TSharedPtr<FSchemaAssetParameters>>
 {
 public:
 	/** Editable text widget stored in this row to easily trigger edit mode */
 	TSharedPtr<SInlineEditableTextBlock> TextBlock;
+
+	/** Object entry box for picking possible Template sequences, and stored in this row to easily trigger edit mode. */
+	TSharedPtr<SObjectPropertyEntryBox> TemplateObjectEntry;
 };
 
 /** An entry in the content folder tree view with knowledge of its path, children, and parent */
@@ -53,8 +67,14 @@ struct FSchemaTreeItem
 	/** The children of this item in the tree that are Folder types */
 	TArray<TSharedPtr<FSchemaTreeItem>> ChildFolders;
 
+	/** Asset Template soft path when this item represents an asset. */
+	FSoftObjectPath AssetTemplate;
+
 	/** The text widget that displays this item's name and supports renaming */
 	TSharedPtr<SInlineEditableTextBlock> NameWidget;
+
+	/** Object entry box for picking possible Template sequences and supports changing the template. */
+	TSharedPtr<SObjectPropertyEntryBox> TemplateObjectEntry;
 };
 
 /** Drag and Drop operation to handle dragging schema assets and folders  */
@@ -68,7 +88,7 @@ public:
 
 	FSchemaAssetDragDrop() = default;
 
-	static TSharedRef<FSchemaAssetDragDrop> New(const FString& InName);
+	static TSharedRef<FSchemaAssetDragDrop> New(const FSchemaAssetParameters& InParameters);
 
 	// Begin FDragDropOperation overrides
 	virtual void OnDrop(bool bDropWasHandled, const FPointerEvent& MouseEvent) override;
@@ -82,6 +102,9 @@ public:
 	/** The name of the item being displayed, used by the decorator and to name the new tree item that will get made when dropped */
 	FString ItemName;
 
+	/** Template object path. */
+	FSoftObjectPath Template;
+
 	/** Callback to handle the case where an item is dropped somewhere not on the tree view */
 	FOnDropNotHanlded OnDropNotHandled;
 };
@@ -164,7 +187,7 @@ private:
 	FReply OnAddFolderClicked();
 
 	/** Generates a row in the asset list view */
-	TSharedRef<ITableRow> OnGenerateAssetRow(TSharedPtr<FString> InItem, const TSharedRef<STableViewBase>& OwnerTable);
+	TSharedRef<ITableRow> OnGenerateAssetRow(TSharedPtr<FSchemaAssetParameters> InItem, const TSharedRef<STableViewBase>& OwnerTable);
 
 	/** Callback when the asset list view rebuilds itself, used to make the textbox in the most recently added row editable */
 	void OnAssetListRebuilt();
@@ -200,7 +223,7 @@ private:
 	void GetFolderListRecursive(TSharedPtr<FSchemaTreeItem> TreeItem, TArray<FString>& FolderList) const;
 
 	/** Recursively get the path of every asset item in the tree view */
-	void GetAssetListRecursive(TSharedPtr<FSchemaTreeItem> TreeItem, TArray<FString>& AssetPathList) const;
+	void GetAssetListRecursive(TSharedPtr<FSchemaTreeItem> TreeItem, TArray<FString>& AssetPathList, TMap<FString, FSoftObjectPath>& Templates) const;
 
 	/** Returns the tree item whose path matches the input path */
 	TSharedPtr<FSchemaTreeItem> FindItemAtPathRecursive(TSharedPtr<FSchemaTreeItem> TreeItem, const FString& Path) const;
@@ -209,11 +232,14 @@ private:
 	void SetChildrenPathRecursive(TSharedPtr<FSchemaTreeItem> TreeItem, const FString& NewPath);
 
 	/** Finds the row for the input item and puts its textblock into edit mode */
-	void EnterEditMode(TSharedPtr<FString> ItemToRename);
+	void EnterEditMode(TSharedPtr<FSchemaAssetParameters> ItemToRename);
 
 	/** Renames the tree item, and updates the paths of all of its children */
 	void OnTreeItemTextCommitted(const FText& InText, ETextCommit::Type InCommitType, TSharedPtr<FSchemaTreeItem> TreeItem);
 
+	/** Updates the Template object of the given tree item. */
+	void OnTreeItemTemplateChanged(const FAssetData& InAssetData, TSharedPtr<FSchemaTreeItem> TreeItem);
+
 	/** Callback when one of the properties of the schema being configured changes */
 	void OnSchemaPropertiesChanged(const FPropertyChangedEvent& PropertyChangedEvent);
 
@@ -224,7 +250,7 @@ private:
 	bool ValidateFolderName(const FText& InText, FText& OutErrorMessage, TSharedPtr<FSchemaTreeItem> TreeItem) const;
 
 	/** Remove the input item from the asset list */
-	void DeleteAssetItem(TSharedPtr<FString> InItem);
+	void DeleteAssetItem(TSharedPtr<FSchemaAssetParameters> InItem);
 
 	/** Remove the input item and its children from the tree view */
 	void DeleteTreeItem(TSharedPtr<FSchemaTreeItem> TreeItem);
@@ -271,10 +297,10 @@ private:
 	FString CreateAssetPath;
 
 	/** List items sources for the asset list view */
-	TArray<TSharedPtr<FString>> AssetListItems;
+	TArray<TSharedPtr<FSchemaAssetParameters>> AssetListItems;
 
 	/** List view of assets that will be associate with the schema */
-	TSharedPtr<SListView<TSharedPtr<FString>>> AssetListView;
+	TSharedPtr<SListView<TSharedPtr<FSchemaAssetParameters>>> AssetListView;
 
 	/** Items source for the tree view */
 	TArray<TSharedPtr<FSchemaTreeItem>> TreeItems;
diff --git a/Engine/Source/Runtime/MovieScene/Private/Bindings/MovieSceneReplaceableBinding.cpp b/Engine/Source/Runtime/MovieScene/Private/Bindings/MovieSceneReplaceableBinding.cpp
--- a/Engine/Source/Runtime/MovieScene/Private/Bindings/MovieSceneReplaceableBinding.cpp
+++ b/Engine/Source/Runtime/MovieScene/Private/Bindings/MovieSceneReplaceableBinding.cpp
@@ -174,4 +174,22 @@ UMovieSceneCustomBinding* UMovieSceneReplaceableBindingBase::CreateNewCustomBind
 	return NewCustomBinding;
 }
 
+void UMovieSceneReplaceableBindingBase::PostDuplicate(EDuplicateMode::Type DuplicateMode)
+{
+	Super::PostDuplicate(DuplicateMode);
+
+#if WITH_EDITOR
+	if (PreviewSpawnable != nullptr)
+	{
+		// Only duplicate the inner PreviewSpawnable if this binding has been bound into a different valid UMovieScene.
+		if (UMovieScene* OwnerMovieScene = GetTypedOuter<UMovieScene>(); 
+			OwnerMovieScene != nullptr && PreviewSpawnable->GetTypedOuter<UMovieScene>() != OwnerMovieScene)
+		{
+			// Duplicate the inner spawnable into the new owning movie scene
+			PreviewSpawnable = Cast<UMovieSceneSpawnableBindingBase>(StaticDuplicateObject(PreviewSpawnable, OwnerMovieScene));
+		}
+	}
+#endif // WITH_EDITOR
+}
+
 #undef LOCTEXT_NAMESPACE
diff --git a/Engine/Source/Runtime/MovieScene/Public/Bindings/MovieSceneReplaceableBinding.h b/Engine/Source/Runtime/MovieScene/Public/Bindings/MovieSceneReplaceableBinding.h
--- a/Engine/Source/Runtime/MovieScene/Public/Bindings/MovieSceneReplaceableBinding.h
+++ b/Engine/Source/Runtime/MovieScene/Public/Bindings/MovieSceneReplaceableBinding.h
@@ -66,6 +66,8 @@ public:
 	UE_API UMovieSceneCustomBinding* CreateNewCustomBinding(UObject* SourceObject, UMovieScene& OwnerMovieScene) override final;
 
 protected:
+	/** UObject overrides */
+	UE_API virtual void PostDuplicate(EDuplicateMode::Type DuplicateMode) override;
 
 	/*
 	* Must be implemented. Called during non-editor/runtime to resolve the binding dynamically. In editor worlds/Sequencer will instead use the PreviewSpawnable binding to spawn a preview object.
diff --git a/Engine/Source/Runtime/MovieSceneTracks/Private/Bindings/MovieSceneSpawnableActorBinding.cpp b/Engine/Source/Runtime/MovieSceneTracks/Private/Bindings/MovieSceneSpawnableActorBinding.cpp
--- a/Engine/Source/Runtime/MovieSceneTracks/Private/Bindings/MovieSceneSpawnableActorBinding.cpp
+++ b/Engine/Source/Runtime/MovieSceneTracks/Private/Bindings/MovieSceneSpawnableActorBinding.cpp
@@ -549,4 +549,27 @@ FText UMovieSceneSpawnableActorBinding::GetBindingTypePrettyName() const
 }
 #endif
 
+void UMovieSceneSpawnableActorBinding::PostDuplicate(EDuplicateMode::Type DuplicateMode)
+{
+	Super::PostDuplicate(DuplicateMode);
+
+	// Post duplication, the Actor Template will still be pointing to the original Sequence.
+	// We therefore need to rebuild the object template
+	if (ActorTemplate != nullptr)
+	{
+		// Only duplicate the inner ActorTemplate if this binding has been bound into a different valid UMovieScene.
+		if (UMovieScene* OuterMovieScene = GetTypedOuter<UMovieScene>(); 
+			OuterMovieScene != nullptr && ActorTemplate->GetTypedOuter<UMovieScene>() != OuterMovieScene)
+		{
+			// Reparent the actor template to the duplicated movie scene
+			ActorTemplate = Cast<AActor>(StaticDuplicateObject(ActorTemplate, OuterMovieScene));
+			if (ActorTemplate)
+			{
+				FMovieSceneSpawnable::MarkSpawnableTemplate(*ActorTemplate);
+			}
+			AutoSetNetAddressableName();
+		}
+	}
+}
+
 #undef LOCTEXT_NAMESPACE
diff --git a/Engine/Source/Runtime/MovieSceneTracks/Private/Bindings/MovieSceneSpawnableDirectorBlueprintBinding.cpp b/Engine/Source/Runtime/MovieSceneTracks/Private/Bindings/MovieSceneSpawnableDirectorBlueprintBinding.cpp
--- a/Engine/Source/Runtime/MovieSceneTracks/Private/Bindings/MovieSceneSpawnableDirectorBlueprintBinding.cpp
+++ b/Engine/Source/Runtime/MovieSceneTracks/Private/Bindings/MovieSceneSpawnableDirectorBlueprintBinding.cpp
@@ -127,4 +127,27 @@ FText UMovieSceneSpawnableDirectorBlueprintBinding::GetBindingTrackIconTooltip()
 }
 #endif
 
+void UMovieSceneSpawnableDirectorBlueprintBinding::PostDuplicate(EDuplicateMode::Type DuplicateMode)
+{
+	Super::PostDuplicate(DuplicateMode);
+
+	// If we were duplicated into a different package we can't reference the old function any more
+	//    For now we null it out, but it would be good to copy the endpoint as well (we currently can't do that
+	//    because there is no way to generically access the correct director BP class from UMovieSceneSequence
+	if (DynamicBinding.Function && GetTypedOuter<UMovieScene>() != nullptr)
+	{
+		UPackage* FunctionPackage = DynamicBinding.Function->GetPackage();
+		UPackage* Package = GetPackage();
+
+		if (Package != FunctionPackage)
+		{
+			DynamicBinding.Function = nullptr;
+			DynamicBinding.ResolveParamsProperty = nullptr;
+		#if WITH_EDITORONLY_DATA
+			DynamicBinding.WeakEndpoint = nullptr;
+		#endif
+		}
+	}
+}
+
 #undef LOCTEXT_NAMESPACE
diff --git a/Engine/Source/Runtime/MovieSceneTracks/Public/Bindings/MovieSceneSpawnableActorBinding.h b/Engine/Source/Runtime/MovieSceneTracks/Public/Bindings/MovieSceneSpawnableActorBinding.h
--- a/Engine/Source/Runtime/MovieSceneTracks/Public/Bindings/MovieSceneSpawnableActorBinding.h
+++ b/Engine/Source/Runtime/MovieSceneTracks/Public/Bindings/MovieSceneSpawnableActorBinding.h
@@ -140,6 +140,8 @@ public:
 #endif
 
 protected:
+	/* UObject overrides */
+	MOVIESCENETRACKS_API virtual void PostDuplicate(EDuplicateMode::Type DuplicateMode) override;
 
 	/* MovieSceneSpawnableBindingBase overrides*/
 	MOVIESCENETRACKS_API UWorld* GetWorldContext(TSharedRef<const UE::MovieScene::FSharedPlaybackState> SharedPlaybackState) const override;
diff --git a/Engine/Source/Runtime/MovieSceneTracks/Public/Bindings/MovieSceneSpawnableDirectorBlueprintBinding.h b/Engine/Source/Runtime/MovieSceneTracks/Public/Bindings/MovieSceneSpawnableDirectorBlueprintBinding.h
--- a/Engine/Source/Runtime/MovieSceneTracks/Public/Bindings/MovieSceneSpawnableDirectorBlueprintBinding.h
+++ b/Engine/Source/Runtime/MovieSceneTracks/Public/Bindings/MovieSceneSpawnableDirectorBlueprintBinding.h
@@ -50,6 +50,8 @@ protected:
 	FText GetBindingTrackIconTooltip() const override;
 #endif
 
+	/** UObject overrides */
+	MOVIESCENETRACKS_API virtual void PostDuplicate(EDuplicateMode::Type DuplicateMode) override;
 
 private:
 	/* MovieSceneSpawnableBindingBase overrides*/
