ShotManagement: Adding ExtendedData to CinematicProduction structure and adding functionality to register and unregister ExtendedData structure types.

diff --git a/Engine/Plugins/VirtualProduction/CinematicAssemblyTools/Source/CineAssemblyToolsEditor/CineAssemblyToolsEditor.build.cs b/Engine/Plugins/VirtualProduction/CinematicAssemblyTools/Source/CineAssemblyToolsEditor/CineAssemblyToolsEditor.build.cs
--- a/Engine/Plugins/VirtualProduction/CinematicAssemblyTools/Source/CineAssemblyToolsEditor/CineAssemblyToolsEditor.build.cs
+++ b/Engine/Plugins/VirtualProduction/CinematicAssemblyTools/Source/CineAssemblyToolsEditor/CineAssemblyToolsEditor.build.cs
@@ -44,6 +44,7 @@ namespace UnrealBuildTool.Rules
 					"Slate",
 					"SlateCore",
 					"SourceControl",
+					"StructUtilsEditor",
 					"TimeManagement",
 					"TakeRecorder",
 					"TakesCore",
diff --git a/Engine/Plugins/VirtualProduction/CinematicAssemblyTools/Source/CineAssemblyToolsEditor/Private/CineAssemblyToolsEditorModule.cpp b/Engine/Plugins/VirtualProduction/CinematicAssemblyTools/Source/CineAssemblyToolsEditor/Private/CineAssemblyToolsEditorModule.cpp
--- a/Engine/Plugins/VirtualProduction/CinematicAssemblyTools/Source/CineAssemblyToolsEditor/Private/CineAssemblyToolsEditorModule.cpp
+++ b/Engine/Plugins/VirtualProduction/CinematicAssemblyTools/Source/CineAssemblyToolsEditor/Private/CineAssemblyToolsEditorModule.cpp
@@ -92,6 +92,16 @@ void FCineAssemblyToolsEditorModule::StartupModule()
 	}
 }
 
+void FCineAssemblyToolsEditorModule::RegisterProductionExtension(const UScriptStruct& DataScriptStruct)
+{
+	ProductionExtensions.RegisterProductionExtension(DataScriptStruct);
+}
+
+void FCineAssemblyToolsEditorModule::UnregisterProductionExtension(const UScriptStruct& DataScriptStruct)
+{
+	ProductionExtensions.UnregisterProductionExtension(DataScriptStruct);
+}
+
 void FCineAssemblyToolsEditorModule::OnPostEngineInit()
 {
 	TakeRecorderIntegration = MakeUnique<FCineAssemblyTakeRecorderIntegration>();
diff --git a/Engine/Plugins/VirtualProduction/CinematicAssemblyTools/Source/CineAssemblyToolsEditor/Private/CineAssemblyToolsEditorModule.h b/Engine/Plugins/VirtualProduction/CinematicAssemblyTools/Source/CineAssemblyToolsEditor/Private/CineAssemblyToolsEditorModule.h
--- a/Engine/Plugins/VirtualProduction/CinematicAssemblyTools/Source/CineAssemblyToolsEditor/Private/CineAssemblyToolsEditorModule.h
+++ b/Engine/Plugins/VirtualProduction/CinematicAssemblyTools/Source/CineAssemblyToolsEditor/Private/CineAssemblyToolsEditorModule.h
@@ -2,8 +2,9 @@
 
 #pragma once
 
+#include "ICineAssemblyToolsEditorModule.h"
 #include "Modules/ModuleManager.h"
-
+#include "ProductionExtensions.h"
 #include "TakeRecorder/CineAssemblyTakeRecorderIntegration.h"
 #include "Widgets/Docking/SDockTab.h"
 
@@ -15,7 +16,7 @@ class UMoviePipelineExecutorJob;
 class ULevelSequence;
 struct FCinematicProduction;
 
-class FCineAssemblyToolsEditorModule : public IModuleInterface
+class FCineAssemblyToolsEditorModule : public ICineAssemblyToolsEditorModule
 {
 public:
 
@@ -27,6 +28,14 @@ public:
 	/** Opens a window to edit the properties of the input assembly schema */
 	void OpenSchemaForEdit(UCineAssemblySchema* Schema);
 
+	/** Get the ProductionExtensions object. */
+	inline const FProductionExtensions& GetProductionExtensions() { return ProductionExtensions; }
+
+	//~ Begin ICineAssemblyToolsEditorModule interface
+	virtual void RegisterProductionExtension(const UScriptStruct& DataScriptStruct) override;
+	virtual void UnregisterProductionExtension(const UScriptStruct& DataScriptStruct) override;
+	//~ End ICineAssemblyToolsEditorModule interface
+
 private:
 	void OnPostEngineInit();
 
@@ -77,6 +86,9 @@ private:
 	/** Manages all integration of Cinematic Assembly Tools with Take Recorder. */
 	TUniquePtr<FCineAssemblyTakeRecorderIntegration> TakeRecorderIntegration;
 
+	/** Extensions for the Production data and Production Wizard. */
+	FProductionExtensions ProductionExtensions;
+
 	/** Map of TabIds to Asset IDs to track which assets are open in which tabs */
 	TMap<FTabId, FGuid> ManagedAssemblyTabs;
 	TMap<FTabId, FGuid> ManagedSchemaTabs;
diff --git a/Engine/Plugins/VirtualProduction/CinematicAssemblyTools/Source/CineAssemblyToolsEditor/Private/ProductionExtensions.cpp b/Engine/Plugins/VirtualProduction/CinematicAssemblyTools/Source/CineAssemblyToolsEditor/Private/ProductionExtensions.cpp
new file mode 100644
--- /dev/null
+++ b/Engine/Plugins/VirtualProduction/CinematicAssemblyTools/Source/CineAssemblyToolsEditor/Private/ProductionExtensions.cpp
@@ -0,0 +1,25 @@
+// Copyright Epic Games, Inc. All Rights Reserved.
+
+#include "ProductionExtensions.h"
+
+#include "CineAssemblyToolsEditorModule.h"
+
+DEFINE_LOG_CATEGORY_STATIC(LogProductionExtensions, Log, All)
+
+void FProductionExtensions::RegisterProductionExtension(const UScriptStruct& DataScriptStruct)
+{
+	if (!DataScriptStruct.IsInBlueprint())
+	{
+		// We do not allow or support blueprint generated structs
+		ProductionExtensions.Add(&DataScriptStruct);
+	}
+	else
+	{
+		UE_LOG(LogProductionExtensions, Warning, TEXT("Blueprint structs are not supported as Production Extended Data types."))
+	}
+}
+
+void FProductionExtensions::UnregisterProductionExtension(const UScriptStruct& DataScriptStruct)
+{
+	ProductionExtensions.Remove(&DataScriptStruct);
+}
diff --git a/Engine/Plugins/VirtualProduction/CinematicAssemblyTools/Source/CineAssemblyToolsEditor/Private/ProductionExtensions.h b/Engine/Plugins/VirtualProduction/CinematicAssemblyTools/Source/CineAssemblyToolsEditor/Private/ProductionExtensions.h
new file mode 100644
--- /dev/null
+++ b/Engine/Plugins/VirtualProduction/CinematicAssemblyTools/Source/CineAssemblyToolsEditor/Private/ProductionExtensions.h
@@ -0,0 +1,24 @@
+// Copyright Epic Games, Inc. All Rights Reserved.
+
+#pragma once
+
+#include "Containers/Set.h"
+#include "ICineAssemblyToolsEditorModule.h"
+
+/** Class that allows registering of extensions to Productions and the Production Wizard. */
+class FProductionExtensions
+{
+public:
+	/** Register a new production extension struct type. */
+	void RegisterProductionExtension(const UScriptStruct& DataScriptStruct);
+
+	/** Unregister a Production extension struct type. */
+	void UnregisterProductionExtension(const UScriptStruct& DataScriptStruct);
+
+	/** Gets all registered Production Extension types. */
+	inline const TSet<const UScriptStruct*>& GetProductionExtensions() const { return ProductionExtensions; }
+
+private:
+	/** Set of all registered Production Extension types. */
+	TSet<const UScriptStruct*> ProductionExtensions;
+};
diff --git a/Engine/Plugins/VirtualProduction/CinematicAssemblyTools/Source/CineAssemblyToolsEditor/Private/ProductionFunctionLibrary.cpp b/Engine/Plugins/VirtualProduction/CinematicAssemblyTools/Source/CineAssemblyToolsEditor/Private/ProductionFunctionLibrary.cpp
--- a/Engine/Plugins/VirtualProduction/CinematicAssemblyTools/Source/CineAssemblyToolsEditor/Private/ProductionFunctionLibrary.cpp
+++ b/Engine/Plugins/VirtualProduction/CinematicAssemblyTools/Source/CineAssemblyToolsEditor/Private/ProductionFunctionLibrary.cpp
@@ -85,6 +85,29 @@ void UProductionFunctionLibrary::RenameProduction(FGuid ProductionID, FString Ne
 	ProductionSettings->RenameProduction(ProductionID, NewName);
 }
 
+bool UProductionFunctionLibrary::GetProductionExtendedData(FGuid ProductionID, const UScriptStruct* DataStruct, FInstancedStruct& OutData)
+{
+	if (DataStruct == nullptr)
+	{
+		return false;
+	}
+
+	const UProductionSettings* ProductionSettings = GetDefault<UProductionSettings>();
+	if (FConstStructView FoundStruct = ProductionSettings->GetProductionExtendedData(ProductionID, *DataStruct); FoundStruct.IsValid())
+	{
+		OutData.InitializeAs(DataStruct, FoundStruct.GetMemory());
+		return true;
+	}
+
+	return false;
+}
+
+bool UProductionFunctionLibrary::SetProductionExtendedData(FGuid ProductionID, const FInstancedStruct& Data)
+{
+	UProductionSettings* ProductionSettings = GetMutableDefault<UProductionSettings>();
+	return ProductionSettings->SetProductionExtendedData(ProductionID, Data);
+}
+
 UCineAssembly* UProductionFunctionLibrary::CreateAssembly(UCineAssemblySchema* Schema, TSoftObjectPtr<UWorld> Level, TSoftObjectPtr<UCineAssembly> ParentAssembly, TMap<FString, FString> Metadata, const FString& Path, const FString& Name, bool bUseDefaultNameFromSchema)
 {
 	// If the supplied path is a virtual path, strip off the "/All" from the beginning before proceeding
diff --git a/Engine/Plugins/VirtualProduction/CinematicAssemblyTools/Source/CineAssemblyToolsEditor/Private/ProductionFunctionLibrary.h b/Engine/Plugins/VirtualProduction/CinematicAssemblyTools/Source/CineAssemblyToolsEditor/Private/ProductionFunctionLibrary.h
--- a/Engine/Plugins/VirtualProduction/CinematicAssemblyTools/Source/CineAssemblyToolsEditor/Private/ProductionFunctionLibrary.h
+++ b/Engine/Plugins/VirtualProduction/CinematicAssemblyTools/Source/CineAssemblyToolsEditor/Private/ProductionFunctionLibrary.h
@@ -63,6 +63,14 @@ public:
 	UFUNCTION(BlueprintCallable, Category = "Cine Assembly Tools")
 	static void RenameProduction(FGuid ProductionID, FString NewName);
 
+	/** Given a ProductionID and an Extended Data Struct type, returns the Extended Data object for that production if found. */
+	UFUNCTION(BlueprintCallable, Category = "Cine Assembly Tools")
+	static bool GetProductionExtendedData(FGuid ProductionID, const UScriptStruct* DataStruct, FInstancedStruct& OutData);
+
+	/** Given a ProductionID and an Extended Data Struct, sets the struct data on the production if the production already contains that type. */
+	UFUNCTION(BlueprintCallable, Category = "Cine Assembly Tools")
+	static bool SetProductionExtendedData(FGuid ProductionID, const FInstancedStruct& Data);
+
 	/** 
 	 * Create a new CineAssembly asset using the input Schema, Level, and Metadata.
 	 * If bUseDefaultNameFromSchema is true, the default assembly name from the specified schema will be used as the new asset name.
diff --git a/Engine/Plugins/VirtualProduction/CinematicAssemblyTools/Source/CineAssemblyToolsEditor/Private/ProductionSettings.cpp b/Engine/Plugins/VirtualProduction/CinematicAssemblyTools/Source/CineAssemblyToolsEditor/Private/ProductionSettings.cpp
--- a/Engine/Plugins/VirtualProduction/CinematicAssemblyTools/Source/CineAssemblyToolsEditor/Private/ProductionSettings.cpp
+++ b/Engine/Plugins/VirtualProduction/CinematicAssemblyTools/Source/CineAssemblyToolsEditor/Private/ProductionSettings.cpp
@@ -6,14 +6,17 @@
 #include "Algo/IndexOf.h"
 #include "AssetToolsModule.h"
 #include "CineAssemblyToolsAnalytics.h"
+#include "CineAssemblyToolsEditorModule.h"
 #include "Engine/Engine.h"
 #include "HAL/FileManager.h"
 #include "HAL/PlatformFileManager.h"
 #include "LevelSequenceProjectSettings.h"
 #include "Misc/ConfigCacheIni.h"
 #include "Misc/CoreDelegates.h"
+#include "Misc/OutputDeviceNull.h"
 #include "MovieSceneToolsProjectSettings.h"
 #include "NamingTokensEngineSubsystem.h"
+#include "ProductionExtensions.h"
 #include "Sections/MovieSceneSubSection.h"
 #include "SourceControlHelpers.h"
 #include "SSettingsEditorCheckoutNotice.h"
@@ -38,6 +41,139 @@ FCinematicProduction::FCinematicProduction()
 	{
 		ProductionID = FGuid::NewGuid();
 	}
+
+	// Ensure the engine is initialized and config loaded.
+	if (GEngine && GEngine->IsInitialized())
+	{
+		LoadAllExtendedData();
+	}
+}
+
+FInstancedStruct* FCinematicProduction::LoadExtendedData(const UScriptStruct& ForStruct)
+{
+	// Only data types that have been registered in Production Extensions can be loaded into a Production
+	FCineAssemblyToolsEditorModule& CatEdModule = FModuleManager::Get().GetModuleChecked<FCineAssemblyToolsEditorModule>("CineAssemblyToolsEditor");
+
+	if (!CatEdModule.GetProductionExtensions().GetProductionExtensions().Contains(&ForStruct))
+	{
+		return nullptr;
+	}
+
+	// Look to see if we have an exported value to import for this type.
+	FInstancedStruct* Data = nullptr;
+	if (FString* ConfigData = ExtendedDataConfig.Find(ForStruct.GetPathName()))
+	{
+		if (TOptional<FInstancedStruct*> ImportedData = ImportExtendedData(*ConfigData); ImportedData.IsSet())
+		{
+			Data = ImportedData.GetValue();
+		}
+	}
+	else
+	{
+		Data = &AddUniqueExtendedData(ForStruct);
+	}
+
+	if (Data != nullptr)
+	{
+		// Sync up our config representation with any new additions
+		ExportExtendedData(*Data);
+	}
+
+	return Data;
+}
+
+FInstancedStruct* FCinematicProduction::FindOrLoadExtendedData(const UScriptStruct& ForStruct)
+{
+	if (FInstancedStruct* Existing = ExtendedData.Find(ForStruct.GetPathName()))
+	{
+		ExportExtendedData(ForStruct);
+		return Existing;
+	}
+	return LoadExtendedData(ForStruct);
+}
+
+const FInstancedStruct* FCinematicProduction::FindExtendedData(const UScriptStruct& ForStruct) const
+{
+	return ExtendedData.Find(ForStruct.GetPathName());
+}
+
+void FCinematicProduction::ExportExtendedData(const UScriptStruct& ForData)
+{
+	if (FInstancedStruct* DataStruct = ExtendedData.Find(ForData.GetPathName()))
+	{
+		ExportExtendedData(*DataStruct);
+	}
+}
+
+const FName FCinematicProduction::GetExtendedDataMemberName()
+{
+	return GET_MEMBER_NAME_CHECKED(FCinematicProduction, ExtendedData);
+}
+
+void FCinematicProduction::LoadAllExtendedData()
+{
+	FCineAssemblyToolsEditorModule& CatEdModule = FModuleManager::Get().GetModuleChecked<FCineAssemblyToolsEditorModule>("CineAssemblyToolsEditor");
+	
+	FScopedModifyProductionExtendedData ModifyGuard(*this);
+	for (const UScriptStruct* ExtendedDataType : CatEdModule.GetProductionExtensions().GetProductionExtensions())
+	{
+		if (ExtendedDataType != nullptr)
+		{
+			LoadExtendedData(*ExtendedDataType);
+		}
+	}
+}
+
+TOptional<FInstancedStruct*> FCinematicProduction::ImportExtendedData(const FString& ImportText)
+{
+	FInstancedStruct NewStruct;
+	const TCHAR* Buffer = *ImportText;
+
+	FOutputDeviceNull SilenceErrors;
+	NewStruct.ImportTextItem(Buffer, PPF_None, nullptr, &SilenceErrors);
+
+	if (NewStruct.IsValid())
+	{
+		return TOptional<FInstancedStruct*>(&ExtendedData.Add(NewStruct.GetScriptStruct()->GetPathName(), NewStruct));
+	}
+	return TOptional<FInstancedStruct*>();
+}
+
+void FCinematicProduction::ExportExtendedData(const FInstancedStruct& InData)
+{
+	const FSoftObjectPath ObjectPath = InData.GetScriptStruct()->GetPathName();
+	if (FString* ExistingConfig = ExtendedDataConfig.Find(ObjectPath))
+	{
+		ExistingConfig->Empty();
+		InData.ExportTextItem(*ExistingConfig, InData, nullptr, PPF_None, nullptr);
+	}
+	else
+	{
+		FString NewData;
+		InData.ExportTextItem(NewData, InData, nullptr, PPF_None, nullptr);
+		ExtendedDataConfig.Add(ObjectPath, NewData);
+	}
+	NotifyExtendedDataExported(InData.GetScriptStruct());
+}
+
+FInstancedStruct& FCinematicProduction::AddUniqueExtendedData(const UScriptStruct& ScriptStruct)
+{
+	if (FInstancedStruct* Existing = ExtendedData.Find(ScriptStruct.GetPathName()))
+	{
+		return *Existing;
+	}
+
+	FInstancedStruct Data;
+	Data.InitializeAs(&ScriptStruct);
+	return ExtendedData.Add(ScriptStruct.GetPathName(), Data);
+}
+
+void FCinematicProduction::NotifyExtendedDataExported(const UScriptStruct* ExportedStruct)
+{
+	if (!bBlockNotifyExtendedDataExported)
+	{
+		ExtendedDataExported.Broadcast(*this, ExportedStruct);
+	}
 }
 
 FName UProductionSettings::GetContainerName() const
@@ -94,14 +230,19 @@ void UProductionSettings::PostInitProperties()
 		ApplyProjectOverrides();
 	}
 
+	FCoreDelegates::OnPostEngineInit.AddUObject(this, &UProductionSettings::OnPostEngineInit);
+}
+
+void UProductionSettings::OnPostEngineInit()
+{
 	// Register a filter with the naming tokens subsystem to apply the active production's namespace deny list when evaluating tokens
-	auto RegisterNamespaceFilter = [this]()
-		{
-			UNamingTokensEngineSubsystem* NamingTokensSubsystem = GEngine->GetEngineSubsystem<UNamingTokensEngineSubsystem>();
-			NamingTokensSubsystem->RegisterNamespaceFilter("CinematicProductionSettings", FFilterNamespace::CreateUObject(this, &UProductionSettings::FilterNamingTokenNamespaces));
-		};
+	UNamingTokensEngineSubsystem* NamingTokensSubsystem = GEngine->GetEngineSubsystem<UNamingTokensEngineSubsystem>();
+	NamingTokensSubsystem->RegisterNamespaceFilter("CinematicProductionSettings", FFilterNamespace::CreateUObject(this, &UProductionSettings::FilterNamingTokenNamespaces));
 
-	FCoreDelegates::OnPostEngineInit.AddLambda(RegisterNamespaceFilter);
+	// Build and load extended production data
+	BindToProductions();
+	ReloadProductionsExtendedData();
+	ProductionListChangedDelegate.AddLambda([this]() { BindToProductions(); });
 }
 
 void UProductionSettings::CacheProjectDefaults()
@@ -114,6 +255,75 @@ void UProductionSettings::CacheProjectDefaults()
 	ProjectDefaultStartTime = MovieSceneToolsSettings->DefaultStartTime;
 }
 
+FStructView UProductionSettings::GetMutableProductionExtendedData(FGuid ProductionID, const UScriptStruct& DataStruct)
+{
+	FStructView Result;
+	if (FCinematicProduction* Production = Algo::FindBy(Productions, ProductionID, &FCinematicProduction::ProductionID))
+	{
+		if (FInstancedStruct* Found = Production->FindOrLoadExtendedData(DataStruct))
+		{
+			Result = *Found;
+		}
+	}
+	return Result;
+}
+
+FConstStructView UProductionSettings::GetProductionExtendedData(FGuid ProductionID, const UScriptStruct& DataStruct) const
+{
+	FConstStructView Result;
+	if (const FCinematicProduction* Production = Algo::FindBy(Productions, ProductionID, &FCinematicProduction::ProductionID))
+	{
+		if (const FInstancedStruct* Found = Production->FindExtendedData(DataStruct))
+		{
+			Result = *Found;
+		}
+	}
+	return Result;
+}
+
+bool UProductionSettings::SetProductionExtendedData(FGuid ProductionID, const FConstStructView& DataStructView)
+{
+	if (!DataStructView.IsValid())
+	{
+		return false;
+	}
+
+	if (FCinematicProduction* Production = Algo::FindBy(Productions, ProductionID, &FCinematicProduction::ProductionID))
+	{
+		if (FInstancedStruct* Found = Production->FindOrLoadExtendedData(*DataStructView.GetScriptStruct()))
+		{
+			FScopedModifyProductionExtendedData ScopedModify(*Production, *DataStructView.GetScriptStruct());
+			*Found = DataStructView;
+			return true;
+		}
+	}
+
+	return false;
+}
+
+FScopedModifyProductionExtendedData UProductionSettings::ModifyProductionExtendedData(FGuid ProductionID)
+{
+	if (FCinematicProduction* Production = Algo::FindBy(Productions, ProductionID, &FCinematicProduction::ProductionID))
+	{
+		return FScopedModifyProductionExtendedData(*Production);
+	}
+	return FScopedModifyProductionExtendedData();
+}
+
+FScopedModifyProductionExtendedData UProductionSettings::ModifyProductionExtendedData(FGuid ProductionID, const UScriptStruct& ForStruct)
+{
+	if (FCinematicProduction* Production = Algo::FindBy(Productions, ProductionID, &FCinematicProduction::ProductionID))
+	{
+		return FScopedModifyProductionExtendedData(*Production, ForStruct);
+	}
+	return FScopedModifyProductionExtendedData();
+}
+
+FName UProductionSettings::GetProductionsPropertyMemberName() 
+{
+	return GET_MEMBER_NAME_CHECKED(UProductionSettings, Productions); 
+}
+
 void UProductionSettings::ApplyProjectOverrides()
 {
 	// Modify the property flags and tooltip of these overridden project settings
@@ -284,6 +494,46 @@ void UProductionSettings::PostEditChangeProperty(struct FPropertyChangedEvent& P
 	}
 }
 
+void UProductionSettings::PostEditChangeChainProperty(FPropertyChangedChainEvent& PropertyChangedEvent)
+{
+	Super::PostEditChangeChainProperty(PropertyChangedEvent);
+
+	// Is this a Production change?
+	if (TDoubleLinkedList<FProperty*>::TDoubleLinkedListNode* ProductionsNode = PropertyChangedEvent.PropertyChain.GetHead();
+		ProductionsNode != nullptr && ProductionsNode->GetValue()->GetFName() == GET_MEMBER_NAME_CHECKED(UProductionSettings, Productions))
+	{
+		// Check if this is is an ExtendedData change.
+		if (TDoubleLinkedList<FProperty*>::TDoubleLinkedListNode* ExtendedDataNode = ProductionsNode->GetNextNode();
+			ExtendedDataNode && ExtendedDataNode->GetValue()->GetFName() == FCinematicProduction::GetExtendedDataMemberName())
+		{
+			const int32 ProductionIndex = PropertyChangedEvent.GetArrayIndex(ProductionsNode->GetValue()->GetName());
+
+			if (Productions.IsValidIndex(ProductionIndex))
+			{
+				if (const UScriptStruct* OwnerStruct = PropertyChangedEvent.Property->GetTypedOwner<UScriptStruct>())
+				{
+					Productions[ProductionIndex].ExportExtendedData(*OwnerStruct);
+				}
+			}
+		}
+	}
+}
+
+void UProductionSettings::PostReloadConfig(FProperty* PropertyThatWasLoaded)
+{
+	Super::PostReloadConfig(PropertyThatWasLoaded);
+	// Reload all productions extended data
+	ReloadProductionsExtendedData();
+}
+
+void UProductionSettings::PostEditUndo()
+{
+	Super::PostEditUndo();
+	// Make sure bindings are valid as the undo will recreate struct instances
+	BindToProductions();
+	ReloadProductionsExtendedData();
+}
+
 void UProductionSettings::AddProduction()
 {
 	FCinematicProduction NewProduction;
@@ -344,6 +594,41 @@ void UProductionSettings::UpdateConfig()
 	}
 }
 
+void UProductionSettings::ReloadProductionsExtendedData()
+{
+	// For every registered extended data, attempt to load that into each production
+	FCineAssemblyToolsEditorModule& CatEdModule = FModuleManager::Get().GetModuleChecked<FCineAssemblyToolsEditorModule>("CineAssemblyToolsEditor");
+	const TSet<const UScriptStruct*>& Extensions = CatEdModule.GetProductionExtensions().GetProductionExtensions();
+
+	for (FCinematicProduction& Production : Productions)
+	{
+		FScopedModifyProductionExtendedData ModifyGuard(Production);
+		for (const UScriptStruct* Struct : Extensions)
+		{
+			if (Struct != nullptr)
+			{
+				Production.LoadExtendedData(*Struct);
+			}
+		}
+	}
+}
+
+void UProductionSettings::BindToProductions()
+{
+	for (FCinematicProduction& Production : Productions)
+	{
+		if (!Production.OnExtendedDataExported().IsBoundToObject(this))
+		{
+			Production.OnExtendedDataExported().AddUObject(this, &UProductionSettings::HandleProductionExtendedDataUpdated);
+		}
+	}
+}
+
+void UProductionSettings::HandleProductionExtendedDataUpdated(const FCinematicProduction& SenderProduction, const UScriptStruct* ForData)
+{
+	UpdateConfig();
+}
+
 void UProductionSettings::DeleteProduction(FGuid ProductionID)
 {
 	if (!ProductionID.IsValid())
diff --git a/Engine/Plugins/VirtualProduction/CinematicAssemblyTools/Source/CineAssemblyToolsEditor/Private/ProductionSettings.h b/Engine/Plugins/VirtualProduction/CinematicAssemblyTools/Source/CineAssemblyToolsEditor/Private/ProductionSettings.h
deleted file mode 100644
diff --git a/Engine/Plugins/VirtualProduction/CinematicAssemblyTools/Source/CineAssemblyToolsEditor/Private/ScopedModifyProductionExtendedData.cpp b/Engine/Plugins/VirtualProduction/CinematicAssemblyTools/Source/CineAssemblyToolsEditor/Private/ScopedModifyProductionExtendedData.cpp
new file mode 100644
--- /dev/null
+++ b/Engine/Plugins/VirtualProduction/CinematicAssemblyTools/Source/CineAssemblyToolsEditor/Private/ScopedModifyProductionExtendedData.cpp
@@ -0,0 +1,81 @@
+// Copyright Epic Games, Inc. All Rights Reserved.
+
+#include "ScopedModifyProductionExtendedData.h"
+
+#include "CineAssemblyToolsEditorModule.h"
+#include "ProductionSettings.h"
+
+FScopedModifyProductionExtendedData::FScopedModifyProductionExtendedData(FCinematicProduction& InProduction) :
+	TargetProduction(InProduction)
+{
+	Start();
+}
+
+FScopedModifyProductionExtendedData::FScopedModifyProductionExtendedData(FCinematicProduction& InProduction, const UScriptStruct& InTargetStruct) :
+	TargetProduction(InProduction), TargetStruct(&InTargetStruct)
+{
+	Start();
+}
+
+FScopedModifyProductionExtendedData::~FScopedModifyProductionExtendedData()
+{
+	Finish();
+}
+
+bool FScopedModifyProductionExtendedData::IsActive() const
+{
+	if (bHasFinished)
+	{
+		return false;
+	}
+
+	if (!TargetStruct.IsSet())
+	{
+		return TargetProduction.IsValid();
+	}
+	return TargetProduction.IsValid() && TargetProduction->FindExtendedData(*TargetStruct.GetValue());
+}
+
+void FScopedModifyProductionExtendedData::Finish()
+{
+	const UScriptStruct* ExportedStruct = nullptr;
+	bool bExported = false;
+
+	if (IsActive())
+	{
+		if (TargetStruct.IsSet())
+		{
+			ExportedStruct = TargetStruct.GetValue();
+			TargetProduction->ExportExtendedData(*ExportedStruct);
+		}
+		else
+		{
+			FCineAssemblyToolsEditorModule& CatEdModule = FModuleManager::Get().GetModuleChecked<FCineAssemblyToolsEditorModule>("CineAssemblyToolsEditor");
+			const TSet<const UScriptStruct*>& Extensions = CatEdModule.GetProductionExtensions().GetProductionExtensions();
+			for (const UScriptStruct* Struct : Extensions)
+			{
+				if (Struct != nullptr)
+				{
+					TargetProduction->ExportExtendedData(*Struct);
+				}
+			}
+		}
+		bExported = true;
+	}
+
+	// Reset the block flag
+	TargetProduction->bBlockNotifyExtendedDataExported = bOriginalBlockValue;
+
+	if (bExported)
+	{
+		TargetProduction->NotifyExtendedDataExported(ExportedStruct);
+	}
+	bHasFinished = true;
+}
+
+void FScopedModifyProductionExtendedData::Start()
+{
+	bHasFinished = false;
+	bOriginalBlockValue = TargetProduction->bBlockNotifyExtendedDataExported;
+	TargetProduction->bBlockNotifyExtendedDataExported = true;
+}
diff --git a/Engine/Plugins/VirtualProduction/CinematicAssemblyTools/Source/CineAssemblyToolsEditor/Public/ICineAssemblyToolsEditorModule.h b/Engine/Plugins/VirtualProduction/CinematicAssemblyTools/Source/CineAssemblyToolsEditor/Public/ICineAssemblyToolsEditorModule.h
new file mode 100644
--- /dev/null
+++ b/Engine/Plugins/VirtualProduction/CinematicAssemblyTools/Source/CineAssemblyToolsEditor/Public/ICineAssemblyToolsEditorModule.h
@@ -0,0 +1,58 @@
+// Copyright Epic Games, Inc. All Rights Reserved.
+
+#pragma once
+
+#include "Modules/ModuleInterface.h"
+#include "Modules/ModuleManager.h"
+
+/**
+ * Interface for CineAssemblyToolsEditor module
+ */
+class ICineAssemblyToolsEditorModule : public IModuleInterface
+{
+public:
+	/** Get the implementation instance for the module. */
+	static ICineAssemblyToolsEditorModule& Get()
+	{
+		return FModuleManager::Get().GetModuleChecked<ICineAssemblyToolsEditorModule>("CineAssemblyToolsEditor");
+	}
+
+	/** Gets if this module is loaded. */
+	static bool IsAvailable()
+	{
+		return FModuleManager::Get().IsModuleLoaded("CineAssemblyToolsEditor");
+	}
+
+	/** Virtual destructor. */
+	virtual ~ICineAssemblyToolsEditorModule() = default;
+
+	/** 
+	 * Register a new production extension struct type. 
+	 * 
+	 * An ExtendedData production extension type is a USTRUCT type that is registered as additional data to each Production object.
+	 * The struct must be a standard USTRUCT and registered to the productions using:
+	 * 
+	 * ```
+	 * USTRUCT()
+	 * struct FMyStruct 
+	 * { 
+	 *		GENERATED_BODY()
+	 * 
+	 *		UPROPERTY(EditAnywhere, Category=Default)
+	 *		int MyIntMember;
+	 * }
+	 * 
+	 * {
+	 *		ICineAssemblyToolsEditorModule& CatEdModule = ICineAssemblyToolsEditorModule::Get();
+	 *		CatEdModule.RegisterProductionExtension(*MyStruct::StaticStruct());
+	 * }
+	 * ```
+	 * 
+	 * Once registered, each production will have its own instance of the extended data struct type.
+	 * See FCinematicProduction for more on how to access and modify the data.
+	 */
+	virtual void RegisterProductionExtension(const UScriptStruct& DataScriptStruct) = 0;
+
+	/** Unregister a Production extension struct type. */
+	virtual void UnregisterProductionExtension(const UScriptStruct& DataScriptStruct) = 0;
+};
diff --git a/Engine/Plugins/VirtualProduction/CinematicAssemblyTools/Source/CineAssemblyToolsEditor/Public/ProductionSettings.h b/Engine/Plugins/VirtualProduction/CinematicAssemblyTools/Source/CineAssemblyToolsEditor/Public/ProductionSettings.h
new file mode 100644
--- /dev/null
+++ b/Engine/Plugins/VirtualProduction/CinematicAssemblyTools/Source/CineAssemblyToolsEditor/Public/ProductionSettings.h
@@ -0,0 +1,487 @@
+// Copyright Epic Games, Inc. All Rights Reserved.
+#pragma  once
+
+#include "Engine/DeveloperSettings.h"
+
+#include "Misc/FrameRate.h"
+#include "ScopedModifyProductionExtendedData.h"
+#include "StructUtils/InstancedStruct.h"
+#include "StructUtils/StructView.h"
+
+#include "ProductionSettings.generated.h"
+
+#define UE_API CINEASSEMBLYTOOLSEDITOR_API
+
+DECLARE_MULTICAST_DELEGATE(FOnProductionListChanged);
+DECLARE_MULTICAST_DELEGATE(FOnActiveProductionChanged);
+DECLARE_MULTICAST_DELEGATE_TwoParams(FOnProductionExtendedDataExported, const FCinematicProduction& /* SenderProduction */, const UScriptStruct* /* ExportedStruct */);
+
+/**
+ * Options for determining the hierarchical bias of subsequences
+ */
+UENUM()
+enum class ESubsequencePriority : uint8
+{
+	TopDown,
+	BottomUp
+};
+
+/**
+ * Properties of a folder in the production's template folder hierarchy
+ */
+USTRUCT()
+struct FFolderTemplate
+{
+	GENERATED_BODY()
+
+	UPROPERTY(config, EditAnywhere, Category = "Default")
+	FString InternalPath;
+
+	UPROPERTY(config, EditAnywhere, Category = "Default")
+	bool bCreateIfMissing = true;
+};
+
+/**
+ * Collection of production settings to override the project/editor behavior
+ */
+USTRUCT(BlueprintType)
+struct FCinematicProduction
+{
+	GENERATED_BODY()
+
+	friend class FScopedModifyProductionExtendedData;
+
+public:
+	FCinematicProduction();
+
+	/** Unique ID of the production */
+	UPROPERTY(BlueprintReadOnly, Category = "Default", meta = (IgnoreForMemberInitializationTest))
+	FGuid ProductionID;
+
+	/** Production Name */
+	UPROPERTY(config, EditAnywhere, BlueprintReadOnly, Category = "Default")
+	FString ProductionName;
+
+	/** The default frame rate to set for new Level Sequences */
+	UPROPERTY(config, EditAnywhere, BlueprintReadOnly, Category = "Default")
+	FFrameRate DefaultDisplayRate;
+
+	/** The default frame number (using the default frame rate) that new Level Sequences should start at */
+	UPROPERTY(config, EditAnywhere, BlueprintReadOnly, Category = "Default")
+	int32 DefaultStartFrame;
+
+	/** Controls whether subsequences override parent sequences, or vice versa */
+	UPROPERTY(config, EditAnywhere, BlueprintReadOnly, Category = "Default")
+	ESubsequencePriority SubsequencePriority = ESubsequencePriority::BottomUp;
+
+	/** List of Naming Token namespaces that should not be evaluated */
+	UPROPERTY(config, EditAnywhere, BlueprintReadOnly, Category = "Default")
+	TSet<FString> NamingTokenNamespaceDenyList;
+
+	/** List of default names for specific asset types */
+	UPROPERTY(config, VisibleAnywhere, BlueprintReadOnly, Category = "Default")
+	TMap<TObjectPtr<const UClass>, FString> DefaultAssetNames;
+
+	/** List of folder paths that represent a template folder hierarchy to be used for this production */
+	UPROPERTY(config, EditAnywhere, Category = "Default")
+	TArray<FFolderTemplate> TemplateFolders;
+
+	//~ Extended Data
+public:
+	/**
+	* Attempts to load Extended Data of the given type, and update the data in the production.
+	* This function always loads the data from ExtendedDataConfig if found, regardless of it has already been loaded or not.
+	* 
+	* The given Struct type must have been registered with ICineAssemblyToolsEditorModule as a ProductionExtension.
+	*
+	* Returns a pointer to the loaded struct or nullptr if the type is not registered with ICineAssemblyToolsEditorModule.
+	* The returned pointer is only valid until the next allocation of ExtendedData.
+	*/
+	UE_API FInstancedStruct* LoadExtendedData(const UScriptStruct& ForStruct);
+
+	/**
+	* Attempts to find Extended Data of the given type, loading it if required and updating the data in the production. If the data is already loaded,
+	* then the existing data is returned and not loaded.
+	* 
+	* The given Struct type must have been registered with ICineAssemblyToolsEditorModule as a ProductionExtension.
+	* 
+	* To modify the returned FInstancedStruct and ensure that it is correctly re-exported back into the production after modification so that it is correctly saved into
+	* the config, use FScopedModifyProductionExtendedData.
+	* 
+	* e.g.
+	*
+	* ```
+	* if (FInstancedStruct* Data = Production.FindOrLoadExtendedData(*MySTruct::StaticStruct()))
+	* {
+	*		FScopedModifyProductionExtendedData ModifyGuard(Production, *MyStruct::StaticStruct());
+	*		MyStruct& MyStructRef = Data->GetMutable();
+	* 
+	*		...
+	*		// Modify MyStructRef
+	*		...
+	* }
+	* ```
+	* 
+	* Returns a pointer to the loaded struct or nullptr if the type is not registered with ICineAssemblyToolsEditorModule.
+	* The returned pointer is only valid until the next allocation of ExtendedData.
+	*/
+	UE_API FInstancedStruct* FindOrLoadExtendedData(const UScriptStruct& ForStruct);
+
+	/**
+	* Attempts to find existing Extended Data of the given type. If the data does not exist or has not been loaded, returns a nullptr.
+	* Use FindOrLoadExtendedData if you are unsure if the data has already been loaded, or ensure LoadExtendedData has already been called for this struct type.
+	* 
+	* The returned pointer is null if the extended data was not found, and if found, is only valid until the next allocation of ExtendedData.
+	*/
+	UE_API const FInstancedStruct* FindExtendedData(const UScriptStruct& ForStruct) const;
+
+	/**
+	 * Export the extended data ready to save to Config by exporting the given struct type into ExtendedDataConfig.
+	 * The ExtendedDataConfig map is not emptied as part of this process.
+	 */
+	UE_API void ExportExtendedData(const UScriptStruct& ForData);
+
+	/** Delegate called when the one or more ExtendedData structs are exported to ExtendedDataConfig. */
+	FOnProductionExtendedDataExported& OnExtendedDataExported() { return ExtendedDataExported; }
+
+	/** Helper function to get the compile time checked member name of the ExtendedData property. */
+	UE_API static const FName GetExtendedDataMemberName();
+
+protected:
+	/**
+	 * Instance specific ExtendedData. This is saved to the config by exporting to text.
+	 * Ensure modifications are guarded by a FScopedModifyProductionExtendedData instance or a call to ExportExtendedData.
+	 */
+	UPROPERTY(Transient, SkipSerialization, EditAnywhere, Category = "Default", NoClear, EditFixedSize, meta = (DisplayThumbnail = "false"))
+	TMap<FSoftObjectPath, FInstancedStruct> ExtendedData;
+
+private:
+	/**
+	 * Loads all registered ExtendedData from config loaded data and adds any additionally registered extended data types.
+	 * Exports the data back into the ExtendedDataConfig struct once loading is complete.
+	 */
+	void LoadAllExtendedData();
+
+	/** Import the text as ExtendedData. Returns the the newly created instance if import was successful. */
+	TOptional<FInstancedStruct*> ImportExtendedData(const FString& ImportText);
+
+	/** Export the InstancedStruct into ExtendedDataConfig. */
+	void ExportExtendedData(const FInstancedStruct& InData);
+
+	/** Uniquely add an instance of the given struct type, returning a reference to either the existing data or the newly created data. */
+	FInstancedStruct& AddUniqueExtendedData(const UScriptStruct& ScriptStruct);
+
+	/** Broadcast the ExtendedDataExported if bNotifyExtendedDataExported is true with the given exported struct type. */
+	void NotifyExtendedDataExported(const UScriptStruct* ExportedStruct);
+
+	/** Delegate called when the one or more ExtendedData structs are exported to ExtendedDataConfig. */
+	FOnProductionExtendedDataExported ExtendedDataExported;
+
+	/** Whether the ExtendedDataExported delegate should be blocked or triggered when the ExtendedData is exported. */
+	bool bBlockNotifyExtendedDataExported = false;
+
+	/** 
+	 * Property used to serialize the extended data to configuration without requiring the types are present.
+	 * 
+	 * The text representation is stored independently of the live ExtendedData to handle ExtendedData types that may not be able
+	 * to be loaded due to coming from unloaded or missing modules.
+	 * 
+	 * The map is keyed with a FSoftObjectPath which represents the PathName of the data UScriptStruct.
+	 */
+	UPROPERTY(config)
+	TMap<FSoftObjectPath, FString> ExtendedDataConfig;
+};
+
+/**
+ * Cinematic Production Settings
+ */
+UCLASS(config=Engine, defaultconfig)
+class UProductionSettings : public UDeveloperSettings
+{
+	GENERATED_BODY()
+
+public:
+	UProductionSettings() = default;
+
+	//~ Begin UDeveloperSettings overrides
+	virtual FName GetContainerName() const override;
+	virtual FName GetCategoryName() const override;
+	virtual FName GetSectionName() const override;
+	//~ End UDeveloperSettings overrides
+
+	//~ Begin UObject overrides
+	virtual void PostInitProperties() override;
+	virtual void PostEditChangeProperty(struct FPropertyChangedEvent& PropertyChangedEvent) override;
+	virtual void PostEditChangeChainProperty(struct FPropertyChangedChainEvent& PropertyChangedEvent) override;
+	virtual void PostReloadConfig(FProperty* PropertyThatWasLoaded) override;
+	virtual void PostEditUndo() override;
+	//~ End UObject overrides
+
+	/** The customization is allowed access to the private properties */
+	friend class FProductionSettingsCustomization;
+
+	/** Returns a copy of the production list */
+	const TArray<FCinematicProduction> GetProductions() const;
+
+	/** Returns a copy of the production matching the input production ID (if it exists) */
+	TOptional<const FCinematicProduction> GetProduction(FGuid ProductionID) const;
+
+	/** Adds a new empty production to the list */
+	void AddProduction();
+
+	/** Adds the input production to the list */
+	void AddProduction(const FCinematicProduction& ProductionToAdd);
+
+	/** Adds a duplicate of the input production to the list */
+	void DuplicateProduction(FGuid ProductionID);
+
+	/** Removes the production matching the input ID from the list */
+	void DeleteProduction(FGuid ProductionID);
+
+	/** Renames the production matching the input ID */
+	void RenameProduction(FGuid ProductionID, FString NewName);
+
+	/** Returns a copy of the active production (if there is one) */
+	TOptional<const FCinematicProduction> GetActiveProduction() const;
+
+	/** Returns the unique ID of the active production */
+	UE_API FGuid GetActiveProductionID() const;
+
+	/** Sets the active production based on the input production ID */
+	void SetActiveProduction(FGuid ProductionID);
+
+	/** Returns true if the input production ID matches the ID of the active production */
+	bool IsActiveProduction(FGuid ProductionID) const;
+
+	/** Returns the DefaultDisplayRate of the active production, or the underlying level sequence setting if there is no active production */
+	FFrameRate GetActiveDisplayRate() const;
+
+	/** Returns the DefaultStartFrame of the active production, or the value based on the underlying movie scene tools setting if there is no active production */
+	int32 GetActiveStartFrame() const;
+
+	/** Returns the SubsequencePriority of the active production, or the default config value if there is no active production */
+	ESubsequencePriority GetActiveSubsequencePriority() const;
+
+	/** Sets the DefaultDisplayRate of the production matching the input ID */
+	void SetDisplayRate(FGuid ProductionID, FFrameRate DisplayRate);
+
+	/** Sets the DefaultStartFrame of the production matching the input ID */
+	void SetStartFrame(FGuid ProductionID, int32 StartFrame);
+
+	/** Sets the SubsequencePriority of the production matching the input ID */
+	void SetSubsequencePriority(FGuid ProductionID, ESubsequencePriority Priority);
+
+	/** Adds a Naming Token namespace to the DenyList of the production matching the input ID */
+	void AddNamespaceToDenyList(FGuid ProductionID, const FString& Namespace);
+
+	/** Removes a Naming Token namespace from the DenyList of the production matching the input ID */
+	void RemoveNamespaceFromDenyList(FGuid ProductionID, const FString& Namespace);
+
+	/** Adds a new entry into the DefaultAssetNames map of the production matching the input ID */
+	void AddAssetNaming(FGuid ProductionID, const UClass* AssetClass, const FString& DefaultName);
+
+	/** Removes an entry from the DefaultAssetNames map of the production matching the input ID */
+	void RemoveAssetNaming(FGuid ProductionID, const UClass* AssetClass);
+
+	/** Add a new path to the input production's list of template folders */
+	void AddTemplateFolder(FGuid ProductionID, const FString& Path, bool bCreateIfMissing=true);
+
+	/** Removes a path from the input production's list of template folders */
+	void RemoveTemplateFolder(FGuid ProductionID, const FString& Path);
+
+	/** Sets the input production's template folder hierarchy to the input array of template folders */
+	void SetTemplateFolderHierarchy(FGuid ProductionID, const TArray<FFolderTemplate>& TemplateHierarchy);
+
+	/** Sets the DefaultDisplayRate of the active production */
+	void SetActiveDisplayRate(FFrameRate DisplayRate);
+
+	/** Sets the DefaultStartFrame of the active production */
+	void SetActiveStartFrame(int32 StartFrame);
+
+	/** Sets the SubsequencePriority of the active production */
+	void SetActiveSubsequencePriority(ESubsequencePriority Priority);
+
+	/** Returns a new unique production name */
+	FString GetUniqueProductionName() const;
+	FString GetUniqueProductionName(const FString& BaseName) const;
+
+	/** Checks if the level editor toolbar is enabled. */
+	bool IsLevelEditorToolbarEnabled() const { return bEnableProductionsLevelEditorToolbar; }
+	
+	/** Returns the delegate that broadcasts when a production is added/removed */
+	FOnProductionListChanged& OnProductionListChanged() { return ProductionListChangedDelegate; }
+
+	/** Returns the delegate that broadcasts when the active production changes */
+	FOnActiveProductionChanged& OnActiveProductionChanged() { return ActiveProductionChangedDelegate; }
+
+	/** Given a production ID and data struct, tries to get a mutable view on the productions extended data of the matching type. */
+	template <typename TStruct>
+	TStructView<TStruct> GetMutableProductionExtendedData(FGuid ProductionID)
+	{
+		if (FStructView StructView = GetMutableProductionExtendedData(ProductionID, *TStruct::StaticStruct()); StructView.IsValid())
+		{
+			return TStructView<TStruct>(StructView.Get<TStruct>());
+		}
+		return TStructView<TStruct>();
+	}
+
+	/** Given a production ID and data struct, tries to get a mutable view on the productions extended data of the matching type. */
+	UE_API FStructView GetMutableProductionExtendedData(FGuid ProductionID, const UScriptStruct& DataStruct);
+
+	/** Given a production ID and data struct, tries to get a view on the productions extended data of the matching type. */
+	template <typename TStruct>
+	TConstStructView<TStruct> GetProductionExtendedData(FGuid ProductionID) const
+	{
+		if (FConstStructView StructView = GetProductionExtendedData(ProductionID, *TStruct::StaticStruct()); StructView.IsValid())
+		{
+			return TConstStructView<TStruct>(StructView.Get<const TStruct>());
+		}
+		return TConstStructView<TStruct>();
+	}
+
+	/** Given a production ID and data struct, tries to get a view on the productions extended data of the matching type. */
+	UE_API FConstStructView GetProductionExtendedData(FGuid ProductionID, const UScriptStruct& DataStruct) const;
+
+	/** Given a production ID and struct view, tries to set the struct in the productions extended data of the matching type. */
+	UE_API bool SetProductionExtendedData(FGuid ProductionID, const FConstStructView& DataStructView);
+
+	/** 
+	 * Get a FScopedModifyProductionExtendedData instance for the production matching the given production ID. This will allow modifying
+	 * multiple extended data of a configured Production, ensuring modifications are exported correctly, ready to be saved back to the config with the
+	 * new updates.
+	 * 
+	 * e.g.
+	 * 
+	 * ```
+	 * const FGuid TargetProductionId; // Set from elsewhere or use GetActiveProductionId()
+	 * UProductionSettings* ProductionSettings = GetMutableDefault<UProductionSettings>();
+	 * FScopedModifyProductionExtendedData ModifyGuard = ProductionSettings->ModifyProductionExtendedData(TargetProductionId);
+	 * 
+	 * if (TStructView<MyStruct> DataView1 = ProductionSettings->GetMutableProductionExtendedData<MyStruct>(TargetProductionId); DataView1.IsValid())
+	 * {
+	 *		MyStruct& MyStructRef = DataView1.Get();
+	 *		...
+	 *		// Modify MyStructRef
+	 *		...
+	 * }
+	 * 
+	 * if (TStructView<MyOtherStruct> DataView2 = ProductionSettings->GetMutableProductionExtendedData<MyOtherStruct>(TargetProductionId); DataView2.IsValid())
+	 * {
+	 *		MyOtherStruct& MyOtherStructRef = DataView2.Get();
+	 *		...
+	 *		// Modify MyOtherStructRef
+	 *		...
+	 * }
+	 * ```
+	 * 
+	 * For blueprints, use nodes GetProductionExtendedData (ProductionFunctionLibrary) and SetProductionExtendedData (ProductionFunctionLibrary).
+	 * 
+	 */
+	[[nodiscard]] UE_API FScopedModifyProductionExtendedData ModifyProductionExtendedData(FGuid ProductionID);
+
+	/** Get a FScopedModifyProductionExtendedData instance for the production matching the given production ID and extended data struct type. This will allow modifying
+	 * the targeted extended data of a configured Production, ensuring modifications are exported correctly, ready to be saved back to the config with the
+	 * new updates.
+	 * 
+	 * e.g.
+	 * 
+	 * ```
+	 * const FGuid TargetProductionId; // Set from elsewhere or use GetActiveProductionId()
+	 * UProductionSettings* ProductionSettings = GetMutableDefault<UProductionSettings>();
+	 * if (TStructView<MyStruct> DataView = ProductionSettings->GetMutableProductionExtendedData<MyStruct>(TargetProductionId); DataView.IsValid())
+	 * {
+	 *		FScopedModifyProductionExtendedData ModifyGuard = ProductionSettings->ModifyProductionExtendedData(TargetProductionId);
+	 *		MyStruct& MyStructRef = DataView.Get();
+	 * 
+	 *		...
+	 *		// Modify MyStructRef
+	 *		...
+	 * }
+	 * ```
+	 * 
+	 * For blueprints, use nodes GetProductionExtendedData (ProductionFunctionLibrary) and SetProductionExtendedData (ProductionFunctionLibrary).
+	 * 
+	 */
+	[[nodiscard]] UE_API FScopedModifyProductionExtendedData ModifyProductionExtendedData(FGuid ProductionID, const UScriptStruct& ExtendedDataStruct);
+
+	/** Get the checked member name for the Productions property. */
+	UE_API static FName GetProductionsPropertyMemberName();
+
+private:
+	/** Applies overrides to various project settings based on the active production settings */
+	void ApplyProjectOverrides();
+
+	/** Overrides the DefaultDisplayRate in the level sequence project settings based on the active production setting */
+	void OverrideDefaultDisplayRate();
+
+	/** Overrides the DefaultStartTime in the movie scene tools project settings based on the active production settings */
+	void OverrideDefaultStartTime();
+
+	/** Writes out the Hierarchical bias value (based on the Subsequence Priority setting) to EditorPerProjectUserSettings.ini */
+	void OverrideSubsequenceHierarchicalBias();
+
+	/** Overrides the DefaultAssetNames property of the asset tools project settings based on the active production setting */
+	void OverrideDefaultAssetNames();
+
+	/** Apply the active production's namespace deny list to the input set of namespace names */
+	void FilterNamingTokenNamespaces(TSet<FString>& Namespaces);
+
+	/** Set the active production name and write out the new value to EditorPerProjectUserSettings.ini */
+	void SetActiveProductionName();
+
+	/** Cache the default project settings that are overridden by the active production, used to reset when there is no active production */
+	void CacheProjectDefaults();
+
+	/** Perform any setup after the engine has finished initilization. */
+	void OnPostEngineInit();
+
+	/** Returns a pointer to the active production (if it exists) */
+	const FCinematicProduction* GetActiveProductionPtr() const;
+
+	/** Try to update the default config file (will attempt to make the file writable if needed) */
+	void UpdateConfig();
+
+	/** Reload all Production extended data into the current Productions. */
+	void ReloadProductionsExtendedData();
+
+	/** Binds to events on all productions. */
+	void BindToProductions();
+
+	/** Handler for when a Productions ExtendedData is updated. */
+	void HandleProductionExtendedDataUpdated(const FCinematicProduction& SenderProduction, const UScriptStruct* ForData);
+
+private:
+	/** Name of the active production */
+	UPROPERTY(VisibleAnywhere, Category = "Default", meta = (DisplayName = "Active Production"))
+	FString ActiveProductionName;
+
+	/** List of available productions in this project */
+	UPROPERTY(config, EditAnywhere, Category = "Default")
+	TArray<FCinematicProduction> Productions;
+
+	/** Enable the production level editor toolbar. */
+	UPROPERTY(config, EditAnywhere, Category = "Default", meta = (ConfigRestartRequired = true, DisplayName = "Enable Production Toolbar"))
+	bool bEnableProductionsLevelEditorToolbar = true;
+
+	/** ID of the active production (in the Productions array) */
+	FGuid ActiveProductionID;
+
+	/** Cached default project settings that are overridden by the active production, used to reset when there is no active production */
+	FString ProjectDefaultDisplayRate;
+	float ProjectDefaultStartTime;
+
+	/** Default asset names previously registered with AssetTools, used to reset them when the active production changes */
+	TMap<const UClass*, FString> ProjectDefaultAssetNames;
+
+	/** Original tooltip text for sequencer settings */
+	FString OriginalDefaultDisplayRateTooltip;
+	FString OriginalDefaultStartTimeTooltip;
+
+	/** Delegate that broadcasts when a production is added/removed */
+	FOnProductionListChanged ProductionListChangedDelegate;
+
+	/** Delegate that broadcasts when the active production changes */
+	FOnActiveProductionChanged ActiveProductionChangedDelegate;
+};
+
+#undef UE_API
diff --git a/Engine/Plugins/VirtualProduction/CinematicAssemblyTools/Source/CineAssemblyToolsEditor/Public/ScopedModifyProductionExtendedData.h b/Engine/Plugins/VirtualProduction/CinematicAssemblyTools/Source/CineAssemblyToolsEditor/Public/ScopedModifyProductionExtendedData.h
new file mode 100644
--- /dev/null
+++ b/Engine/Plugins/VirtualProduction/CinematicAssemblyTools/Source/CineAssemblyToolsEditor/Public/ScopedModifyProductionExtendedData.h
@@ -0,0 +1,76 @@
+// Copyright Epic Games, Inc. All Rights Reserved.
+
+#pragma once
+
+#include "StructUtils/StructView.h"
+
+struct FCinematicProduction;
+
+/**
+* Helper class for scoping modifications to a Productions extended data.
+* On destruction of the instance, Finish is called which will call update on the TargetProduction if it has not been already.
+* Finish can also be called at any time if desired, and will not then be called on destruction.
+*
+* Use of this class is to create a locally scoped instance before making any modifications, and allowing it to fall out of scope
+* or manually calling Finish when you wish those updates to propogate.
+*
+* eg.
+* ```
+* {
+*	if (FInstancedStruct* ExtendedData = Production.FindOrLoadExtendedData(*FMyStruct::StaticStruct()))
+*	{
+*		FScopedModifyProductionExtendedData ModifyGuard(Production);
+*		...
+*		// Modify ExtendedData
+*		...
+*	}
+* }
+* ```
+*
+* You can specify a specific ExtendedData UScriptStruct type to call update only for that type. If this is not provided at construction time,
+* then all ExtendedData is updated on Finish.
+*/
+class CINEASSEMBLYTOOLSEDITOR_API FScopedModifyProductionExtendedData
+{
+public:
+	/** Default constructor that constructs into an inactive state. */
+	FScopedModifyProductionExtendedData() = default;
+
+	/** Given a reference to a production, create a scoped modify instance. */
+	FScopedModifyProductionExtendedData(FCinematicProduction& InProduction);
+
+	/**
+	* Given a reference to a production and a specific data struct type, create a scoped modify instance.
+	*
+	* If InTargetStruct is not found as an ExtendedData type within the Production, the scoped modify instance is not considered active.
+	*/
+	FScopedModifyProductionExtendedData(FCinematicProduction& InProduction, const UScriptStruct& InTargetStruct);
+
+	/** Destructor that will call Finish if the instance is still active. */
+	virtual ~FScopedModifyProductionExtendedData();
+
+	/**
+	* Returns if this instance is active.
+	* It is considered active if it has not yet finished and has a valid Production reference and target data struct if supplied.
+	*/
+	bool IsActive() const;
+
+	/** Finish this scoped modify and update the extended data on the target production. */
+	void Finish();
+
+private:
+	/** Original value of the block flag for the production. */
+	bool bOriginalBlockValue = false;
+
+	/** Start the modification and blocking of exported signals. */
+	void Start();
+
+	/** View onto the target production object that is being modified. */
+	TStructView<FCinematicProduction> TargetProduction;
+
+	/** Optional target struct type that is being updated as part of this modification. */
+	TOptional<const UScriptStruct*> TargetStruct;
+
+	/** Whether this instance has had Finish called on it or not. */
+	bool bHasFinished = false;
+};
