ShotManagement: Add extensions and customizations to the Production Wizard interface for both ProductionSettings and UserSettings.

diff --git a/Engine/Plugins/VirtualProduction/CinematicAssemblyTools/Source/CineAssemblyToolsEditor/Private/CineAssemblyToolsEditorModule.cpp b/Engine/Plugins/VirtualProduction/CinematicAssemblyTools/Source/CineAssemblyToolsEditor/Private/CineAssemblyToolsEditorModule.cpp
--- a/Engine/Plugins/VirtualProduction/CinematicAssemblyTools/Source/CineAssemblyToolsEditor/Private/CineAssemblyToolsEditorModule.cpp
+++ b/Engine/Plugins/VirtualProduction/CinematicAssemblyTools/Source/CineAssemblyToolsEditor/Private/CineAssemblyToolsEditorModule.cpp
@@ -102,6 +102,27 @@ void FCineAssemblyToolsEditorModule::UnregisterProductionExtension(const UScript
 	ProductionExtensions.UnregisterProductionExtension(DataScriptStruct);
 }
 
+void FCineAssemblyToolsEditorModule::RegisterProductionWizardCustomization(const UScriptStruct& ForDataScriptStruct, FGetWidget MakeCustomWidget, 
+	TAttribute<FText> Label, TAttribute<FSlateIcon> Icon, bool bShowProductionSelection, bool bHideInWizard)
+{
+	ProductionExtensions.RegisterProductionWizardCustomization(ForDataScriptStruct, MakeCustomWidget, Label, Icon, bShowProductionSelection, bHideInWizard);
+}
+
+void FCineAssemblyToolsEditorModule::UnregisterProductionWizardCustomization(const UScriptStruct& ForDataScriptStruct)
+{
+	ProductionExtensions.UnregisterProductionWizardCustomization(ForDataScriptStruct);
+}
+
+void FCineAssemblyToolsEditorModule::RegisterProductionWizardUserSettings(const FName Name, FGetWidget MakeCustomWidget, TAttribute<FText> Label, TAttribute<FSlateIcon> Icon)
+{
+	ProductionExtensions.RegisterProductionWizardUserSettings(Name, MakeCustomWidget, Label, Icon);
+}
+
+void FCineAssemblyToolsEditorModule::UnregisterProductionWizardUserSettings(const FName Name)
+{
+	ProductionExtensions.UnregisterProductionWizardUserSettings(Name);
+}
+
 void FCineAssemblyToolsEditorModule::OnPostEngineInit()
 {
 	TakeRecorderIntegration = MakeUnique<FCineAssemblyTakeRecorderIntegration>();
diff --git a/Engine/Plugins/VirtualProduction/CinematicAssemblyTools/Source/CineAssemblyToolsEditor/Private/CineAssemblyToolsEditorModule.h b/Engine/Plugins/VirtualProduction/CinematicAssemblyTools/Source/CineAssemblyToolsEditor/Private/CineAssemblyToolsEditorModule.h
--- a/Engine/Plugins/VirtualProduction/CinematicAssemblyTools/Source/CineAssemblyToolsEditor/Private/CineAssemblyToolsEditorModule.h
+++ b/Engine/Plugins/VirtualProduction/CinematicAssemblyTools/Source/CineAssemblyToolsEditor/Private/CineAssemblyToolsEditorModule.h
@@ -34,6 +34,11 @@ public:
 	//~ Begin ICineAssemblyToolsEditorModule interface
 	virtual void RegisterProductionExtension(const UScriptStruct& DataScriptStruct) override;
 	virtual void UnregisterProductionExtension(const UScriptStruct& DataScriptStruct) override;
+	virtual void RegisterProductionWizardCustomization(const UScriptStruct& ForDataScriptStruct,
+		FGetWidget MakeCustomWidget, TAttribute<FText> Label, TAttribute<FSlateIcon> Icon, bool bShowProductionSelection, bool bHideInWizard) override;
+	virtual void UnregisterProductionWizardCustomization(const UScriptStruct& ForDataScriptStruct) override;
+	virtual void RegisterProductionWizardUserSettings(const FName Name, FGetWidget MakeCustomWidget, TAttribute<FText> Label, TAttribute<FSlateIcon> Icon) override;
+	virtual void UnregisterProductionWizardUserSettings(const FName Name) override;
 	//~ End ICineAssemblyToolsEditorModule interface
 
 private:
diff --git a/Engine/Plugins/VirtualProduction/CinematicAssemblyTools/Source/CineAssemblyToolsEditor/Private/ProductionExtensions.cpp b/Engine/Plugins/VirtualProduction/CinematicAssemblyTools/Source/CineAssemblyToolsEditor/Private/ProductionExtensions.cpp
--- a/Engine/Plugins/VirtualProduction/CinematicAssemblyTools/Source/CineAssemblyToolsEditor/Private/ProductionExtensions.cpp
+++ b/Engine/Plugins/VirtualProduction/CinematicAssemblyTools/Source/CineAssemblyToolsEditor/Private/ProductionExtensions.cpp
@@ -23,3 +23,40 @@ void FProductionExtensions::UnregisterProductionExtension(const UScriptStruct& D
 {
 	ProductionExtensions.Remove(&DataScriptStruct);
 }
+
+void FProductionExtensions::RegisterProductionWizardCustomization(const UScriptStruct& ForDataScriptStruct, FGetWidget MakeCustomWidget, 
+	TAttribute<FText> Label, TAttribute<FSlateIcon> Icon, bool bShowProductionSelection, bool bHideInWizard)
+{
+	// only register known extension types
+	if (!ProductionExtensions.Contains(&ForDataScriptStruct))
+	{
+		return;
+	}
+
+	FProductionWizardCustomizationParameters Customization;
+	Customization.Label = Label;
+	Customization.Icon = Icon;
+	Customization.MakeCustomWidget = MakeCustomWidget;
+	Customization.bShowProductionSelection = bShowProductionSelection;
+	Customization.bHideInWizard = bHideInWizard;
+	ProductionWizardCustomizations.Add(&ForDataScriptStruct, Customization);
+}
+
+void FProductionExtensions::UnregisterProductionWizardCustomization(const UScriptStruct& ForDataScriptStruct)
+{
+	ProductionWizardCustomizations.Remove(&ForDataScriptStruct);
+}
+
+void FProductionExtensions::RegisterProductionWizardUserSettings(const FName Name, FGetWidget MakeCustomWidget, TAttribute<FText> Label, TAttribute<FSlateIcon> Icon)
+{
+	FProductionWizardUserSettingsExtensionParameters Context;
+	Context.Label = Label.IsSet() ? Label.Get() : FText::FromString(Name.ToString());
+	Context.MakeCustomWidget = MakeCustomWidget;
+	Context.Icon = Icon;
+	ProductionWizardExtendedUserSettings.Add(Name, Context);
+}
+
+void FProductionExtensions::UnregisterProductionWizardUserSettings(const FName Name)
+{
+	ProductionWizardExtendedUserSettings.Remove(Name);
+}
diff --git a/Engine/Plugins/VirtualProduction/CinematicAssemblyTools/Source/CineAssemblyToolsEditor/Private/ProductionExtensions.h b/Engine/Plugins/VirtualProduction/CinematicAssemblyTools/Source/CineAssemblyToolsEditor/Private/ProductionExtensions.h
--- a/Engine/Plugins/VirtualProduction/CinematicAssemblyTools/Source/CineAssemblyToolsEditor/Private/ProductionExtensions.h
+++ b/Engine/Plugins/VirtualProduction/CinematicAssemblyTools/Source/CineAssemblyToolsEditor/Private/ProductionExtensions.h
@@ -3,8 +3,43 @@
 #pragma once
 
 #include "Containers/Set.h"
+#include "Internationalization/Text.h"
+#include "Misc/Attribute.h"
+#include "Textures/SlateIcon.h"
 #include "ICineAssemblyToolsEditorModule.h"
 
+/** Struct holding details on a Production Settings customization for the Production Wizard interface. */
+struct FProductionWizardCustomizationParameters
+{
+	/** Label override to display on the tab in the Production Wizard. */
+	TAttribute<FText> Label;
+
+	/** Icon override to display on the tab in the Production Wizard. */
+	TAttribute<FSlateIcon> Icon;
+
+	/** Callback to create the SWidget to display in the Production Wizard. */
+	FGetWidget MakeCustomWidget = nullptr;
+
+	/** Whether the production selection combo box should be shown for this customization in the Production Wizard. */
+	bool bShowProductionSelection = true;
+
+	/** Hides the tab and widget completely from the Production Wizard. */
+	bool bHideInWizard = false;
+};
+
+/** Struct holding details on a Users Settings extension for the Production Wizard interface. */
+struct FProductionWizardUserSettingsExtensionParameters
+{
+	/** Label for the tab. */
+	FText Label;
+
+	/** Optional icon to display on the tab. */
+	TAttribute<FSlateIcon> Icon;
+
+	/** Callback to create the widget to display. */
+	FGetWidget MakeCustomWidget = nullptr;
+};
+
 /** Class that allows registering of extensions to Productions and the Production Wizard. */
 class FProductionExtensions
 {
@@ -15,10 +50,40 @@ public:
 	/** Unregister a Production extension struct type. */
 	void UnregisterProductionExtension(const UScriptStruct& DataScriptStruct);
 
+	/**
+	 * Register a customization for the Production Wizard for the given extended data struct type.
+	 * The extended data struct type must already have been registered as a Production extension using RegisterProductionExtension.
+	 */
+	void RegisterProductionWizardCustomization(const UScriptStruct& ForDataScriptStruct,
+		FGetWidget MakeCustomWidget, TAttribute<FText> Label = TAttribute<FText>(),
+		TAttribute<FSlateIcon> Icon = TAttribute<FSlateIcon>(), bool bShowProductionSelection = true, bool bHideInWizard = false);
+
+	/** Unregister the registered customization for the given extended data struct type. */
+	void UnregisterProductionWizardCustomization(const UScriptStruct& ForDataScriptStruct);
+
+	/** Register a custom User Setting tab to show in the Production wizard by providing a callback to create the SWidget to show. */
+	void RegisterProductionWizardUserSettings(const FName Name, FGetWidget MakeCustomWidget,
+		TAttribute<FText> Label = TAttribute<FText>(), TAttribute<FSlateIcon> Icon = TAttribute<FSlateIcon>());
+
+	/** Unregister the custom User Settings tab from the Production Wizard with the given name. */
+	void UnregisterProductionWizardUserSettings(const FName Name);
+
 	/** Gets all registered Production Extension types. */
 	inline const TSet<const UScriptStruct*>& GetProductionExtensions() const { return ProductionExtensions; }
 
+	/** Get registered production wizard customizations. */
+	inline const TMap<const UScriptStruct*, FProductionWizardCustomizationParameters>& GetProductionWizardCustomizations() const { return ProductionWizardCustomizations; }
+
+	/** Get the map of all registered custom Production Wizard User Settings tabs. */
+	inline const TMap<FName, FProductionWizardUserSettingsExtensionParameters>& GetProductionWizardUserSettingsExtensions() const { return ProductionWizardExtendedUserSettings; }
+
 private:
 	/** Set of all registered Production Extension types. */
 	TSet<const UScriptStruct*> ProductionExtensions;
+
+	/** Map of all registered Production Wizard Production Settings customizations. */
+	TMap<const UScriptStruct*, FProductionWizardCustomizationParameters> ProductionWizardCustomizations;
+
+	/** Map of all registered Production Wizard Users Settings extensions. */
+	TMap<FName, FProductionWizardUserSettingsExtensionParameters> ProductionWizardExtendedUserSettings;
 };
diff --git a/Engine/Plugins/VirtualProduction/CinematicAssemblyTools/Source/CineAssemblyToolsEditor/Private/ProductionFunctionLibrary.h b/Engine/Plugins/VirtualProduction/CinematicAssemblyTools/Source/CineAssemblyToolsEditor/Private/ProductionFunctionLibrary.h
--- a/Engine/Plugins/VirtualProduction/CinematicAssemblyTools/Source/CineAssemblyToolsEditor/Private/ProductionFunctionLibrary.h
+++ b/Engine/Plugins/VirtualProduction/CinematicAssemblyTools/Source/CineAssemblyToolsEditor/Private/ProductionFunctionLibrary.h
@@ -67,7 +67,7 @@ public:
 	UFUNCTION(BlueprintCallable, Category = "Cine Assembly Tools")
 	static bool GetProductionExtendedData(FGuid ProductionID, const UScriptStruct* DataStruct, FInstancedStruct& OutData);
 
-	/** Given a ProductionID and an Extended Data Struct, sets the struct data on the production if the production already contains that type. */
+	/** Given a ProductionID and an Extended Data Struct, sets the struct data on the production. */
 	UFUNCTION(BlueprintCallable, Category = "Cine Assembly Tools")
 	static bool SetProductionExtendedData(FGuid ProductionID, const FInstancedStruct& Data);
 
diff --git a/Engine/Plugins/VirtualProduction/CinematicAssemblyTools/Source/CineAssemblyToolsEditor/Private/ScopedModifyProductionExtendedData.cpp b/Engine/Plugins/VirtualProduction/CinematicAssemblyTools/Source/CineAssemblyToolsEditor/Private/ScopedModifyProductionExtendedData.cpp
--- a/Engine/Plugins/VirtualProduction/CinematicAssemblyTools/Source/CineAssemblyToolsEditor/Private/ScopedModifyProductionExtendedData.cpp
+++ b/Engine/Plugins/VirtualProduction/CinematicAssemblyTools/Source/CineAssemblyToolsEditor/Private/ScopedModifyProductionExtendedData.cpp
@@ -61,21 +61,25 @@ void FScopedModifyProductionExtendedData::Finish()
 			}
 		}
 		bExported = true;
-	}
 
-	// Reset the block flag
-	TargetProduction->bBlockNotifyExtendedDataExported = bOriginalBlockValue;
+		// Reset the block flag
+		TargetProduction->bBlockNotifyExtendedDataExported = bOriginalBlockValue;
 
-	if (bExported)
-	{
-		TargetProduction->NotifyExtendedDataExported(ExportedStruct);
+		if (bExported)
+		{
+			TargetProduction->NotifyExtendedDataExported(ExportedStruct);
+		}
 	}
+
 	bHasFinished = true;
 }
 
 void FScopedModifyProductionExtendedData::Start()
 {
 	bHasFinished = false;
-	bOriginalBlockValue = TargetProduction->bBlockNotifyExtendedDataExported;
-	TargetProduction->bBlockNotifyExtendedDataExported = true;
+	if (TargetProduction.IsValid())
+	{
+		bOriginalBlockValue = TargetProduction->bBlockNotifyExtendedDataExported;
+		TargetProduction->bBlockNotifyExtendedDataExported = true;
+	}
 }
diff --git a/Engine/Plugins/VirtualProduction/CinematicAssemblyTools/Source/CineAssemblyToolsEditor/Private/UI/SActiveProductionExtendedDataDetailsView.cpp b/Engine/Plugins/VirtualProduction/CinematicAssemblyTools/Source/CineAssemblyToolsEditor/Private/UI/SActiveProductionExtendedDataDetailsView.cpp
new file mode 100644
--- /dev/null
+++ b/Engine/Plugins/VirtualProduction/CinematicAssemblyTools/Source/CineAssemblyToolsEditor/Private/UI/SActiveProductionExtendedDataDetailsView.cpp
@@ -0,0 +1,343 @@
+// Copyright Epic Games, Inc. All Rights Reserved.
+
+#include "SActiveProductionExtendedDataDetailsView.h"
+
+#include "DetailCategoryBuilder.h"
+#include "DetailLayoutBuilder.h"
+#include "IDetailCustomization.h"
+#include "InstancedStructDetails.h"
+#include "IStructureDetailsView.h"
+#include "Modules/ModuleManager.h"
+#include "ProductionSettings.h"
+#include "UI/SActiveProductionCombo.h"
+#include "Widgets/Layout/SSeparator.h"
+#include "Widgets/Layout/SWidgetSwitcher.h"
+#include "Widgets/Text/STextBlock.h"
+
+#define LOCTEXT_NAMESPACE "SActiveProductionExtendedDataDetailsView"
+
+namespace UE::ActiveProductionData::Private
+{
+
+constexpr int32 TabIndex_NoProduction = 0;
+constexpr int32 TabIndex_Details = 1;
+
+/** 
+ * Private details customization used specifically for the instance in the SActiveProductionExtendedDataDetailsView.
+ * 
+ * This customization takes the target struct type into account and shows only the Extended Data in the active production
+ * that matches.
+ */
+class FProductionSettingsSingleExtendedDataDetails : public IDetailCustomization
+{
+public:
+	/** Creates an instance of this details customization */
+	static TSharedRef<FProductionSettingsSingleExtendedDataDetails> MakeInstance(const UScriptStruct* ExtendedDataScriptStruct)
+	{
+		TSharedRef<FProductionSettingsSingleExtendedDataDetails> Instance = MakeShared<FProductionSettingsSingleExtendedDataDetails>();
+		Instance->ExtendedDataScriptStruct = ExtendedDataScriptStruct;
+		return Instance;
+	}
+
+	/** Set the target Extended Data struct type. */
+	void SetExtendedDataScriptStruct(const UScriptStruct* DataStruct)
+	{
+		ExtendedDataScriptStruct = DataStruct;
+
+		if (WeakBuilder.IsValid())
+		{
+			WeakBuilder.Pin()->ForceRefreshDetails();
+		}
+	}
+
+protected:
+	//~ Begin IDetailCustomization interface
+	virtual void CustomizeDetails(IDetailLayoutBuilder& DetailBuilder) override
+	{
+		TArray<TWeakObjectPtr<UObject>> ObjectsBeingModified;
+		DetailBuilder.GetObjectsBeingCustomized(ObjectsBeingModified);
+
+		// This should only ever be a single object that is the default UProductionSettings.
+		check(ObjectsBeingModified.Num() == 1);
+		check(ObjectsBeingModified[0] == GetDefault<UProductionSettings>());
+
+		// hide all categories
+		TArray<FName> CategoryNames;
+		DetailBuilder.GetCategoryNames(CategoryNames);
+
+		for (const FName& Category : CategoryNames)
+		{
+			DetailBuilder.HideCategory(Category);
+		}
+
+		if (ExtendedDataScriptStruct == nullptr)
+		{
+			return;
+		}
+
+		// Find the property handle of the specific data we want to display
+		const UProductionSettings* ProductionSettings = Cast<UProductionSettings>(ObjectsBeingModified[0]);
+
+		TSharedPtr<IPropertyHandleArray> ProductionsHandleArray = DetailBuilder.GetProperty(
+			UProductionSettings::GetProductionsPropertyMemberName(), UProductionSettings::StaticClass())->AsArray();
+		check(ProductionsHandleArray);
+
+		const FGuid ActiveProductionId = ProductionSettings->GetActiveProductionID();
+
+		if (TSharedPtr<IPropertyHandle> ActiveProductionHandle = GetProductionPropertyHandle(ActiveProductionId, ProductionsHandleArray))
+		{
+			if (TSharedPtr<IPropertyHandle> TargetDataHandle = GetTargetExtendedDataPropertyHandle(ActiveProductionHandle))
+			{
+				// Add our custom property
+				IDetailCategoryBuilder& CategoryBuilder = DetailBuilder.EditCategory(ExtendedDataScriptStruct->GetFName());
+				TSharedRef<FInstancedStructDataDetails> DataDetails = MakeShared<FInstancedStructDataDetails>(TargetDataHandle);
+				CategoryBuilder.AddCustomBuilder(DataDetails);
+			}
+		}
+	}
+
+	virtual void CustomizeDetails(const TSharedPtr<IDetailLayoutBuilder>& DetailBuilder) override
+	{
+		WeakBuilder = DetailBuilder;
+		IDetailCustomization::CustomizeDetails(DetailBuilder);
+	}
+	//~ End IDetailCustomization interface
+
+private:
+	/** Gets the property handle for the production with the given ID. */
+	TSharedPtr<IPropertyHandle> GetProductionPropertyHandle(FGuid ProductionID, TSharedPtr<IPropertyHandleArray> ProductionsHandleArray)
+	{
+		if (!ProductionID.IsValid())
+		{
+			return nullptr;
+		}
+
+		uint32 NumProductions;
+		ProductionsHandleArray->GetNumElements(NumProductions);
+
+		for (uint32 ProductionIndex = 0; ProductionIndex < NumProductions; ++ProductionIndex)
+		{
+			TArray<void*> RawData;
+			TSharedRef<IPropertyHandle> ProductionHandle = ProductionsHandleArray->GetElement(ProductionIndex);
+			ProductionHandle->AccessRawData(RawData);
+
+			if (!RawData.IsEmpty())
+			{
+				if (FCinematicProduction* ProductionPtr = reinterpret_cast<FCinematicProduction*>(RawData[0]))
+				{
+					if (ProductionPtr->ProductionID == ProductionID)
+					{
+						return ProductionHandle;
+					}
+				}
+			}
+		}
+
+		return nullptr;
+	}
+
+	/** Gets the property handle for the Extended Data struct of the target struct type. */
+	TSharedPtr<IPropertyHandle> GetTargetExtendedDataPropertyHandle(TSharedPtr<IPropertyHandle> ProductionPropertyHandle)
+	{
+		if (TSharedPtr<IPropertyHandle> ExtendedDataHandle = ProductionPropertyHandle->GetChildHandle(FCinematicProduction::GetExtendedDataMemberName()))
+		{
+			TSharedPtr<IPropertyHandleMap> ExtendedDataHandleMap = ExtendedDataHandle->AsMap();
+
+			uint32 NumElements;
+			ExtendedDataHandleMap->GetNumElements(NumElements);
+
+			for (uint32 DataIndex = 0; DataIndex < NumElements; ++DataIndex)
+			{
+				TArray<void*> RawData;
+				TSharedRef<IPropertyHandle> DataHandle = ExtendedDataHandleMap->GetElement(DataIndex);
+				DataHandle->AccessRawData(RawData);
+
+				if (!RawData.IsEmpty())
+				{
+					if (FInstancedStruct* StructPtr = reinterpret_cast<FInstancedStruct*>(RawData[0]))
+					{
+						if (StructPtr->GetScriptStruct() == ExtendedDataScriptStruct)
+						{
+							return DataHandle;
+						}
+					}
+				}
+			}
+		}
+
+		return nullptr;
+	}
+
+	/** Extended Data target struct type. */
+	const UScriptStruct* ExtendedDataScriptStruct = nullptr;
+
+	/** Weak pointer to the layout builder to allow us to call refresh when the target script struct changes. */
+	TWeakPtr<IDetailLayoutBuilder> WeakBuilder;
+};
+
+} // namespace UE::ActiveProductionData::Private
+
+SActiveProductionExtendedDataDetailsView::~SActiveProductionExtendedDataDetailsView()
+{
+	if (UObjectInitialized())
+	{
+		if (UProductionSettings* ProductionSettings = GetMutableDefault<UProductionSettings>())
+		{
+			ProductionSettings->OnProductionListChanged().Remove(ProductionListChangedHandle);
+			ProductionSettings->OnActiveProductionChanged().Remove(ActiveProductionChangedHandle);
+		}
+	}
+}
+
+void SActiveProductionExtendedDataDetailsView::Construct(const FArguments& InArgs)
+{
+	TargetScriptStruct = InArgs._TargetScriptStruct;
+	bShowProductionSelection = InArgs._bShowProductionSelection;
+
+	TSharedPtr<SWidget> Contents = nullptr;
+	if (TargetScriptStruct == nullptr)
+	{
+		// Display a simple text block to inform the user that no target script struct has been set.
+		Contents = SNew(STextBlock)
+			.Justification(ETextJustify::Center)
+			.ToolTipText(LOCTEXT("ActiveProductionDataPanel_NoTargetStructTooltip", "No details can be shown as the target data structure is null."))
+			.Text(LOCTEXT("ActiveProductionDataPanel_NoTargetStructText", "No target data structure set."));
+	}
+	else
+	{
+		// Subscribe to be notified when the Production Settings active productions has changed
+		UProductionSettings* ProductionSettings = GetMutableDefault<UProductionSettings>();
+		ActiveProductionChangedHandle = ProductionSettings->OnActiveProductionChanged().AddSP(this, &SActiveProductionExtendedDataDetailsView::HandleActiveProductionChanged);
+		ProductionListChangedHandle = ProductionSettings->OnProductionListChanged().AddSP(this, &SActiveProductionExtendedDataDetailsView::HandleProductionListChanged);
+		Contents = MakeContents(InArgs._InnerWidget);
+	}
+
+	
+
+	ChildSlot
+		[
+			Contents.ToSharedRef()
+		];
+}
+
+void SActiveProductionExtendedDataDetailsView::HandleActiveProductionChanged()
+{
+	UpdateActiveData();
+}
+
+void SActiveProductionExtendedDataDetailsView::HandleProductionListChanged()
+{
+	UpdateActiveData();
+}
+
+void SActiveProductionExtendedDataDetailsView::UpdateActiveData()
+{
+	if (DetailsView == nullptr || TargetScriptStruct == nullptr)
+	{
+		return;
+	}
+
+	if (WeakDetailsCustomization.IsValid())
+	{
+		WeakDetailsCustomization.Pin()->SetExtendedDataScriptStruct(TargetScriptStruct);
+	}
+
+	UProductionSettings* ProductionSettings = GetMutableDefault<UProductionSettings>();
+	if (ProductionSettings->GetActiveProductionID().IsValid())
+	{
+		DetailsWidgetSwitcher->SetActiveWidgetIndex(UE::ActiveProductionData::Private::TabIndex_Details);
+	}
+	else
+	{
+		DetailsWidgetSwitcher->SetActiveWidgetIndex(UE::ActiveProductionData::Private::TabIndex_NoProduction);
+	}
+}
+
+TSharedRef<SWidget> SActiveProductionExtendedDataDetailsView::MakeContents(TSharedPtr<SWidget> InnerWidget)
+{
+	TSharedRef<SWidget> ContentsWidget = InnerWidget != nullptr
+		? InnerWidget.ToSharedRef()
+		: MakeDetailsWidget();
+
+	if (bShowProductionSelection)
+	{
+		return SNew(SVerticalBox)
+
+		// Active Production Selector
+			+ SVerticalBox::Slot()
+			.AutoHeight()
+			[
+				SNew(SActiveProductionCombo)
+			]
+
+		// Separator
+			+ SVerticalBox::Slot()
+			.AutoHeight()
+			[
+				SNew(SSeparator)
+					.Orientation(Orient_Horizontal)
+					.Thickness(2.0f)
+			]
+
+		// Contents widget
+			+ SVerticalBox::Slot()
+			.FillHeight(1.0f)
+			[
+				ContentsWidget
+			];
+	}
+	
+	return ContentsWidget;
+}
+
+TSharedRef<SWidget> SActiveProductionExtendedDataDetailsView::MakeDetailsWidget()
+{
+	FDetailsViewArgs DetailsViewArgs;
+	DetailsViewArgs.NameAreaSettings = FDetailsViewArgs::HideNameArea;
+	DetailsViewArgs.bShowOptions = false;
+	DetailsViewArgs.bShowPropertyMatrixButton = false;
+
+	FStructureDetailsViewArgs StructureViewArgs;
+	StructureViewArgs.bShowObjects = true;
+	StructureViewArgs.bShowAssets = true;
+	StructureViewArgs.bShowClasses = true;
+	StructureViewArgs.bShowInterfaces = true;
+
+	FPropertyEditorModule& PropModule = FModuleManager::Get().GetModuleChecked<FPropertyEditorModule>("PropertyEditor");
+	DetailsView = PropModule.CreateDetailView(DetailsViewArgs);
+	DetailsView->SetGenericLayoutDetailsDelegate(FOnGetDetailCustomizationInstance::CreateLambda([this]()
+		{
+			// Customize this instance of the details view to only display the single Extended Data. Using a view on the object
+			// allows all change events and undos to propagate back correctly to the object.
+			TSharedRef<UE::ActiveProductionData::Private::FProductionSettingsSingleExtendedDataDetails> DetailsCustomization = 
+				UE::ActiveProductionData::Private::FProductionSettingsSingleExtendedDataDetails::MakeInstance(TargetScriptStruct);
+
+			// Capture the customization object weakly.
+			WeakDetailsCustomization = DetailsCustomization;
+			return DetailsCustomization;
+		}));
+	DetailsView->SetObject(GetMutableDefault<UProductionSettings>());
+
+	DetailsWidgetSwitcher = SNew(SWidgetSwitcher)
+		+ SWidgetSwitcher::Slot()
+		.HAlign(HAlign_Center)
+		.VAlign(VAlign_Center)
+		[
+			// If no production is active, display a simple message to the user to select a production.
+			SNew(STextBlock)
+				.Justification(ETextJustify::Center)
+				.ToolTipText(LOCTEXT("ActiveProductionDataPanel_NoProductionTooltip", "Select a Production to inspect its data."))
+				.Text(LOCTEXT("ActiveProductionDataPanel_NoProductionText", "Select a Production"))
+		]
+
+		+ SWidgetSwitcher::Slot()
+		[
+			DetailsView.ToSharedRef()
+		];
+
+	UpdateActiveData();
+
+	return DetailsWidgetSwitcher.ToSharedRef();
+}
+
+#undef LOCTEXT_NAMESPACE
diff --git a/Engine/Plugins/VirtualProduction/CinematicAssemblyTools/Source/CineAssemblyToolsEditor/Private/UI/SActiveProductionExtendedDataDetailsView.h b/Engine/Plugins/VirtualProduction/CinematicAssemblyTools/Source/CineAssemblyToolsEditor/Private/UI/SActiveProductionExtendedDataDetailsView.h
new file mode 100644
--- /dev/null
+++ b/Engine/Plugins/VirtualProduction/CinematicAssemblyTools/Source/CineAssemblyToolsEditor/Private/UI/SActiveProductionExtendedDataDetailsView.h
@@ -0,0 +1,87 @@
+// Copyright Epic Games, Inc. All Rights Reserved.
+
+#pragma once
+
+#include "Widgets/SCompoundWidget.h"
+#include "Widgets/DeclarativeSyntaxSupport.h"
+
+class IDetailsView;
+class SWidgetSwitcher;
+
+namespace UE::ActiveProductionData::Private
+{
+class FProductionSettingsSingleExtendedDataDetails;
+} // namespace UE::ActiveProductionData::Private
+
+/** 
+ * Widget that shows Extended Data in a details view for the active production. 
+ * 
+ * This widget displays a view of the Extended Data in the context of the UProductionSettings objects.
+ * This means all change events propogate back to the UProductionSettings object.
+ */
+class SActiveProductionExtendedDataDetailsView : public SCompoundWidget
+{
+public:
+	/** Default constructor. */
+	SActiveProductionExtendedDataDetailsView() = default;
+
+	/** Destructor unhooks from UProductionSettings events. */
+	~SActiveProductionExtendedDataDetailsView();
+
+	SLATE_BEGIN_ARGS(SActiveProductionExtendedDataDetailsView) :
+		_TargetScriptStruct(nullptr),
+		_InnerWidget(nullptr),
+		_bShowProductionSelection(true)
+		{}
+
+		/** The target script structure of Extended Data to show. */
+		SLATE_ARGUMENT(const UScriptStruct*, TargetScriptStruct)
+
+		/** A custom widget to display instead of a details view. */
+		SLATE_ARGUMENT(TSharedPtr<SWidget>, InnerWidget)
+
+		/** Whether to display the Production Selection combo box. */
+		SLATE_ARGUMENT(bool, bShowProductionSelection)
+	SLATE_END_ARGS()
+
+	//~ Begin SWidget Interface
+	void Construct(const FArguments& InArgs);
+	//~ End SWidget Interface
+
+private:
+	/** Handles when the active production changes in UProductionSettings. */
+	void HandleActiveProductionChanged();
+
+	/** Handles when the production list changes in UProductionSettings. */
+	void HandleProductionListChanged();
+
+	/** Updates the active and displayed data. */
+	void UpdateActiveData();
+
+	/** Make the contents of the widget taking a custom inner widget and parameters into account. */
+	TSharedRef<SWidget> MakeContents(TSharedPtr<SWidget> InnerWidget);
+
+	/** Makes the details widget and sets up customizations for the view. */
+	TSharedRef<SWidget> MakeDetailsWidget();
+
+	/** The details view in use. */
+	TSharedPtr<IDetailsView> DetailsView;
+
+	/** Switcher widget used to display a message when no active production is selected. */
+	TSharedPtr<SWidgetSwitcher> DetailsWidgetSwitcher;
+
+	/** Forward declared customization object used to customize this instance of the details view. */
+	TWeakPtr<UE::ActiveProductionData::Private::FProductionSettingsSingleExtendedDataDetails> WeakDetailsCustomization;
+
+	/** The Extended Data script structure to display. */
+	const UScriptStruct* TargetScriptStruct = nullptr;
+
+	/** Whether the Productions selection combo box should be shown. */
+	bool bShowProductionSelection = true;
+
+	/** Delegate bound to the Production Setting's OnActiveProductionChanged event */
+	FDelegateHandle ActiveProductionChangedHandle;
+
+	/** Delegate bound to the Production Setting's OnProductionListChanged event */
+	FDelegateHandle ProductionListChangedHandle;
+};
diff --git a/Engine/Plugins/VirtualProduction/CinematicAssemblyTools/Source/CineAssemblyToolsEditor/Private/UI/SProductionWizard.cpp b/Engine/Plugins/VirtualProduction/CinematicAssemblyTools/Source/CineAssemblyToolsEditor/Private/UI/SProductionWizard.cpp
--- a/Engine/Plugins/VirtualProduction/CinematicAssemblyTools/Source/CineAssemblyToolsEditor/Private/UI/SProductionWizard.cpp
+++ b/Engine/Plugins/VirtualProduction/CinematicAssemblyTools/Source/CineAssemblyToolsEditor/Private/UI/SProductionWizard.cpp
@@ -3,8 +3,10 @@
 #include "SProductionWizard.h"
 
 #include "CineAssemblyToolsStyle.h"
+#include "CineAssemblyToolsEditorModule.h"
 #include "ProductionSettings.h"
 #include "UI/SNamingTokensPanel.h"
+#include "UI/SActiveProductionExtendedDataDetailsView.h"
 #include "UI/SAssetNamingPanel.h"
 #include "UI/SFolderHierarchyPanel.h"
 #include "UI/SProductionListPanel.h"
@@ -185,6 +187,85 @@ FReply SProductionWizard::OnNextClicked()
 	return FReply::Handled();
 }
 
+void SProductionWizard::MakeMenuExtensions()
+{
+	// Grab all User Settings extensions to add
+	FCineAssemblyToolsEditorModule& CatEdModule = FModuleManager::Get().GetModuleChecked<FCineAssemblyToolsEditorModule>("CineAssemblyToolsEditor");
+	const FProductionExtensions& Extensions = CatEdModule.GetProductionExtensions();
+
+	for (const TPair<FName, FProductionWizardUserSettingsExtensionParameters>& Entry : Extensions.GetProductionWizardUserSettingsExtensions())
+	{
+		TSharedPtr<SWidget> CustomWidget = nullptr;
+		
+		if (Entry.Value.MakeCustomWidget.IsBound())
+		{
+			CustomWidget = Entry.Value.MakeCustomWidget.Execute();
+		}
+
+		if (CustomWidget == nullptr)
+		{
+			continue;
+		}
+
+		TSharedPtr<FProductionWizardMenuEntry> UserSettingsExtensionEntry = MakeShared<FProductionWizardMenuEntry>();
+		UserSettingsExtensionEntry->Label = Entry.Value.Label;
+
+		if (Entry.Value.Icon.IsSet())
+		{
+			UserSettingsExtensionEntry->Icon = Entry.Value.Icon.Get();
+		}
+
+		UserSettingsExtensionEntry->Panel = CustomWidget;
+
+		UserMenuEntries.Add(MoveTemp(UserSettingsExtensionEntry));
+	}
+
+	// Grab all production settings extensions and customizations to add
+	const TMap<const UScriptStruct*, FProductionWizardCustomizationParameters>& WizardCustomizations = Extensions.GetProductionWizardCustomizations();
+
+	for (const UScriptStruct* ProductionExtensionType : Extensions.GetProductionExtensions())
+	{
+		TSharedPtr<FProductionWizardMenuEntry> ExtensionMenuEntry = MakeShared<FProductionWizardMenuEntry>();
+		ExtensionMenuEntry->Label = FText::FromName(ProductionExtensionType->GetFName());
+
+		bool bShowProductionSelection = true;
+		TSharedPtr<SWidget> InnerWidget = nullptr;
+
+		if (const FProductionWizardCustomizationParameters* Customization = WizardCustomizations.Find(ProductionExtensionType))
+		{
+			// Only add the entry if the customization is not hiding it from the wizard.
+			if (Customization->bHideInWizard)
+			{
+				continue;
+			}
+
+			bShowProductionSelection = Customization->bShowProductionSelection;
+
+			if (Customization->MakeCustomWidget.IsBound())
+			{
+				InnerWidget = Customization->MakeCustomWidget.Execute();
+			}
+
+			if (Customization->Label.IsSet())
+			{
+				ExtensionMenuEntry->Label = Customization->Label.Get();
+			}
+
+			if (Customization->Icon.IsSet())
+			{
+				ExtensionMenuEntry->Icon = Customization->Icon.Get();
+			}
+		}
+
+		ExtensionMenuEntry->Panel = SNew(SActiveProductionExtendedDataDetailsView)
+			.bShowProductionSelection(bShowProductionSelection)
+			.InnerWidget(InnerWidget)
+			.TargetScriptStruct(ProductionExtensionType);
+
+		ProductionMenuEntries.Add(MoveTemp(ExtensionMenuEntry));
+	}
+}
+
 TSharedRef<SWidget> SProductionWizard::MakeMenuPanel()
 {
 	// Build the Revision Control menu entry, including its content panel
@@ -231,6 +312,9 @@ TSharedRef<SWidget> SProductionWizard::MakeMenuPanel()
 	ProductionMenuEntries.Add(AssetNamingMenuEntry);
 	ProductionMenuEntries.Add(FolderHierarchyMenuEntry);
 
+	// Make any registered extensions
+	MakeMenuExtensions();
+
 	// Construct the list view for the user menu entries
 	UserMenuListView = SNew(SListView<TSharedPtr<FProductionWizardMenuEntry>>)
 		.ListItemsSource(&UserMenuEntries)
diff --git a/Engine/Plugins/VirtualProduction/CinematicAssemblyTools/Source/CineAssemblyToolsEditor/Private/UI/SProductionWizard.h b/Engine/Plugins/VirtualProduction/CinematicAssemblyTools/Source/CineAssemblyToolsEditor/Private/UI/SProductionWizard.h
--- a/Engine/Plugins/VirtualProduction/CinematicAssemblyTools/Source/CineAssemblyToolsEditor/Private/UI/SProductionWizard.h
+++ b/Engine/Plugins/VirtualProduction/CinematicAssemblyTools/Source/CineAssemblyToolsEditor/Private/UI/SProductionWizard.h
@@ -57,6 +57,9 @@ private:
 	/** Cycles forwards through the available panels */
 	FReply OnNextClicked();
 
+	/** Adds User Settings and Production Settings extensions to the Menu Panel. */
+	void MakeMenuExtensions();
+
 private:
 	/** Slots for the menu and menu entry content panels */
 	SSplitter::FSlot* MenuPanelSlot = nullptr;
diff --git a/Engine/Plugins/VirtualProduction/CinematicAssemblyTools/Source/CineAssemblyToolsEditor/Public/ICineAssemblyToolsEditorModule.h b/Engine/Plugins/VirtualProduction/CinematicAssemblyTools/Source/CineAssemblyToolsEditor/Public/ICineAssemblyToolsEditorModule.h
--- a/Engine/Plugins/VirtualProduction/CinematicAssemblyTools/Source/CineAssemblyToolsEditor/Public/ICineAssemblyToolsEditorModule.h
+++ b/Engine/Plugins/VirtualProduction/CinematicAssemblyTools/Source/CineAssemblyToolsEditor/Public/ICineAssemblyToolsEditorModule.h
@@ -4,6 +4,14 @@
 
 #include "Modules/ModuleInterface.h"
 #include "Modules/ModuleManager.h"
+#include "Internationalization/Text.h"
+#include "Misc/Attribute.h"
+#include "Textures/SlateIcon.h"
+
+class SWidget;
+
+/** Callback for constructing a new SWidget. */
+DECLARE_DELEGATE_RetVal(TSharedRef<SWidget>, FGetWidget);
 
 /**
  * Interface for CineAssemblyToolsEditor module
@@ -14,7 +22,7 @@ public:
 	/** Get the implementation instance for the module. */
 	static ICineAssemblyToolsEditorModule& Get()
 	{
-		return FModuleManager::Get().GetModuleChecked<ICineAssemblyToolsEditorModule>("CineAssemblyToolsEditor");
+		return FModuleManager::Get().LoadModuleChecked<ICineAssemblyToolsEditorModule>("CineAssemblyToolsEditor");
 	}
 
 	/** Gets if this module is loaded. */
@@ -44,7 +52,7 @@ public:
 	 * 
 	 * {
 	 *		ICineAssemblyToolsEditorModule& CatEdModule = ICineAssemblyToolsEditorModule::Get();
-	 *		CatEdModule.RegisterProductionExtension(*MyStruct::StaticStruct());
+	 *		CatEdModule.RegisterProductionExtension(*FMyStruct::StaticStruct());
 	 * }
 	 * ```
 	 * 
@@ -55,4 +63,31 @@ public:
 
 	/** Unregister a Production extension struct type. */
 	virtual void UnregisterProductionExtension(const UScriptStruct& DataScriptStruct) = 0;
+
+	/**
+	 * Register a customization for the Production Wizard for the given extended data struct type.
+	 * 
+	 * A callback can be provided to construct a custom widget, or leave unbound if the default details view is desired. You can then override the label,
+	 * icon and if the Production Selection combo box is shown in the wizard for the data tab. 
+	 * 
+	 * If bHideInWizard is given as true, then the tab is hidden from the Production Wizard completely for the given struct type.
+	 * 
+	 * The extended data struct type should already have been registered as a Production extension using RegisterProductionExtension.
+	 */
+	virtual void RegisterProductionWizardCustomization(const UScriptStruct& ForDataScriptStruct,
+		FGetWidget MakeCustomWidget, TAttribute<FText> Label, TAttribute<FSlateIcon> Icon, bool bShowProductionSelection, bool bHideInWizard) = 0;
+
+	/** Unregister the registered customization for the given extended data struct type. */
+	virtual void UnregisterProductionWizardCustomization(const UScriptStruct& ForDataScriptStruct) = 0;
+
+	/** 
+	 * Register a custom User Setting tab to show in the Production wizard by providing a callback to create the SWidget to show, as well as customizations to label and icon to display.
+	 * 
+	 * If no label is provided, Name will be used as the display label on the Production Wizard interface.
+	 */
+	virtual void RegisterProductionWizardUserSettings(const FName Name, FGetWidget MakeCustomWidget,
+		TAttribute<FText> Label, TAttribute<FSlateIcon> Icon) = 0;
+
+	/** Unregister the custom User Settings tab from the Production Wizard with the given name. */
+	virtual void UnregisterProductionWizardUserSettings(const FName Name) = 0;
 };
diff --git a/Engine/Plugins/VirtualProduction/CinematicAssemblyTools/Source/CineAssemblyToolsEditor/Public/ProductionSettings.h b/Engine/Plugins/VirtualProduction/CinematicAssemblyTools/Source/CineAssemblyToolsEditor/Public/ProductionSettings.h
--- a/Engine/Plugins/VirtualProduction/CinematicAssemblyTools/Source/CineAssemblyToolsEditor/Public/ProductionSettings.h
+++ b/Engine/Plugins/VirtualProduction/CinematicAssemblyTools/Source/CineAssemblyToolsEditor/Public/ProductionSettings.h
@@ -111,10 +111,10 @@ public:
 	* e.g.
 	*
 	* ```
-	* if (FInstancedStruct* Data = Production.FindOrLoadExtendedData(*MySTruct::StaticStruct()))
+	* if (FInstancedStruct* Data = Production.FindOrLoadExtendedData(*FMyStruct::StaticStruct()))
 	* {
-	*		FScopedModifyProductionExtendedData ModifyGuard(Production, *MyStruct::StaticStruct());
-	*		MyStruct& MyStructRef = Data->GetMutable();
+	*		FScopedModifyProductionExtendedData ModifyGuard(Production, FMyStruct::StaticStruct());
+	*		FMyStruct& MyStructRef = Data->GetMutable();
 	* 
 	*		...
 	*		// Modify MyStructRef
@@ -171,7 +171,7 @@ private:
 	/** Uniquely add an instance of the given struct type, returning a reference to either the existing data or the newly created data. */
 	FInstancedStruct& AddUniqueExtendedData(const UScriptStruct& ScriptStruct);
 
-	/** Broadcast the ExtendedDataExported if bNotifyExtendedDataExported is true with the given exported struct type. */
+	/** Broadcast the ExtendedDataExported if bBlockNotifyExtendedDataExported is false with the given exported struct type. */
 	void NotifyExtendedDataExported(const UScriptStruct* ExportedStruct);
 
 	/** Delegate called when the one or more ExtendedData structs are exported to ExtendedDataConfig. */
diff --git a/Engine/Plugins/VirtualProduction/CinematicAssemblyTools/Source/CineAssemblyToolsEditor/Public/ScopedModifyProductionExtendedData.h b/Engine/Plugins/VirtualProduction/CinematicAssemblyTools/Source/CineAssemblyToolsEditor/Public/ScopedModifyProductionExtendedData.h
--- a/Engine/Plugins/VirtualProduction/CinematicAssemblyTools/Source/CineAssemblyToolsEditor/Public/ScopedModifyProductionExtendedData.h
+++ b/Engine/Plugins/VirtualProduction/CinematicAssemblyTools/Source/CineAssemblyToolsEditor/Public/ScopedModifyProductionExtendedData.h
@@ -46,7 +46,7 @@ public:
 	*/
 	FScopedModifyProductionExtendedData(FCinematicProduction& InProduction, const UScriptStruct& InTargetStruct);
 
-	/** Destructor that will call Finish if the instance is still active. */
+	/** Destructor that will call Finish. */
 	virtual ~FScopedModifyProductionExtendedData();
 
 	/**
