ShotManagement: Removing Iso8601 format for GetCreated as this is a local time, not a formatted Iso8601 time.

#rnx

diff --git a/Engine/Plugins/VirtualProduction/CinematicAssemblyTools/Source/CineAssemblyTools/Private/CineAssembly.cpp b/Engine/Plugins/VirtualProduction/CinematicAssemblyTools/Source/CineAssemblyTools/Private/CineAssembly.cpp
--- a/Engine/Plugins/VirtualProduction/CinematicAssemblyTools/Source/CineAssemblyTools/Private/CineAssembly.cpp
+++ b/Engine/Plugins/VirtualProduction/CinematicAssemblyTools/Source/CineAssemblyTools/Private/CineAssembly.cpp
@@ -516,11 +516,11 @@ void UCineAssembly::SetAuthor(const FString& InAuthor)
 	MetaData->SetAuthor(InAuthor);
 }
 
-FString UCineAssembly::GetCreatedIso8601String() const
+FString UCineAssembly::GetCreatedString() const
 {
 	if (TOptional<FDateTime> CreatedLocalTime = TryGetCreatedAsLocalTime(); CreatedLocalTime.IsSet())
 	{
-		return CreatedLocalTime->ToIso8601();
+		return CreatedLocalTime->ToString();
 	}
 	
 	return FString();
diff --git a/Engine/Plugins/VirtualProduction/CinematicAssemblyTools/Source/CineAssemblyTools/Private/CineAssemblyNamingTokens.cpp b/Engine/Plugins/VirtualProduction/CinematicAssemblyTools/Source/CineAssemblyTools/Private/CineAssemblyNamingTokens.cpp
--- a/Engine/Plugins/VirtualProduction/CinematicAssemblyTools/Source/CineAssemblyTools/Private/CineAssemblyNamingTokens.cpp
+++ b/Engine/Plugins/VirtualProduction/CinematicAssemblyTools/Source/CineAssemblyTools/Private/CineAssemblyNamingTokens.cpp
@@ -99,7 +99,7 @@ void UCineAssemblyNamingTokens::OnCreateDefaultTokens(TArray<FNamingTokenData>&
 		{
 #if WITH_EDITOR
 			return InAssembly.IsValid()
-				? FText::FromString(InAssembly->GetCreatedIso8601String()) 
+				? FText::FromString(InAssembly->GetCreatedString()) 
 				: FText::GetEmpty();
 #else // WITH_EDITOR
 			return FText::GetEmpty();
diff --git a/Engine/Plugins/VirtualProduction/CinematicAssemblyTools/Source/CineAssemblyTools/Public/CineAssembly.h b/Engine/Plugins/VirtualProduction/CinematicAssemblyTools/Source/CineAssemblyTools/Public/CineAssembly.h
--- a/Engine/Plugins/VirtualProduction/CinematicAssemblyTools/Source/CineAssemblyTools/Public/CineAssembly.h
+++ b/Engine/Plugins/VirtualProduction/CinematicAssemblyTools/Source/CineAssemblyTools/Public/CineAssembly.h
@@ -105,9 +105,9 @@ public:
 	UFUNCTION(BlueprintCallable, Category = "Cine Assembly Tools")
 	CINEASSEMBLYTOOLS_API void SetAuthor(const FString& InAuthor);
 
-	/** Gets the created date-time of this assembly as an Iso 8601 format string. */
+	/** Gets the created date-time of this assembly as a localtime formatted string. */
 	UFUNCTION(BlueprintCallable, Category = "Cine Assembly Tools")
-	CINEASSEMBLYTOOLS_API FString GetCreatedIso8601String() const;
+	CINEASSEMBLYTOOLS_API FString GetCreatedString() const;
 
 	/** Gets the date this assembly was created as a YYYY-MM-DD string. */
 	UFUNCTION(BlueprintCallable, Category = "Cine Assembly Tools")
