CineAssembly: Fix cook error from not checking if slate is initialized.

diff --git a/Engine/Plugins/VirtualProduction/CinematicAssemblyTools/Source/CineAssemblyToolsEditor/Private/CineAssemblyToolsEditorModule.cpp b/Engine/Plugins/VirtualProduction/CinematicAssemblyTools/Source/CineAssemblyToolsEditor/Private/CineAssemblyToolsEditorModule.cpp
--- a/Engine/Plugins/VirtualProduction/CinematicAssemblyTools/Source/CineAssemblyToolsEditor/Private/CineAssemblyToolsEditorModule.cpp
+++ b/Engine/Plugins/VirtualProduction/CinematicAssemblyTools/Source/CineAssemblyToolsEditor/Private/CineAssemblyToolsEditorModule.cpp
@@ -183,7 +183,7 @@ void FCineAssemblyToolsEditorModule::ShutdownModule()
 	FCoreDelegates::OnPostEngineInit.RemoveAll(this);
 	TakeRecorderIntegration.Reset();
 
-	if (UObjectInitialized())
+	if (UObjectInitialized() && FSlateApplication::IsInitialized())
 	{
 		UToolMenus::Get()->RemoveSection("LevelEditor.LevelEditorToolBar.User", "Productions");
 	}
@@ -508,7 +508,7 @@ void FCineAssemblyToolsEditorModule::SaveOpenTabs(const TMap<FTabId, FGuid>& Tab
 
 void FCineAssemblyToolsEditorModule::ExtendLevelEditorToolbar()
 {
-	if (!GetDefault<UProductionSettings>()->IsLevelEditorToolbarEnabled())
+	if (!FSlateApplication::IsInitialized() || !GetDefault<UProductionSettings>()->IsLevelEditorToolbarEnabled())
 	{
 		return;
 	}
