ProductionWizard: Disable ability to create new assets from Schema Template selection drop downs

Original reported issue is of the Schema creation window minimizing when trying to create a new Level Sequence or Cine Assembly from the object picker dropdown. However, with further thought this does not feel like a correct workflow anyway as an empty new template is not much use. Therefore, the change here is to remove the ability to create new Template Sequences or Assemblies from those menus entirely.

#rnx

diff --git a/Engine/Plugins/VirtualProduction/CinematicAssemblyTools/Source/CineAssemblyToolsEditor/Private/UI/CineAssembly/SCineAssemblySchemaWindow.cpp b/Engine/Plugins/VirtualProduction/CinematicAssemblyTools/Source/CineAssemblyToolsEditor/Private/UI/CineAssembly/SCineAssemblySchemaWindow.cpp
--- a/Engine/Plugins/VirtualProduction/CinematicAssemblyTools/Source/CineAssemblyToolsEditor/Private/UI/CineAssembly/SCineAssemblySchemaWindow.cpp
+++ b/Engine/Plugins/VirtualProduction/CinematicAssemblyTools/Source/CineAssemblyToolsEditor/Private/UI/CineAssembly/SCineAssemblySchemaWindow.cpp
@@ -749,7 +749,7 @@ TSharedRef<ITableRow> SCineAssemblySchemaWindow::OnGenerateAssetRow(TSharedPtr<F
 		});
 
 	Row->TemplateObjectEntry = SNew(SObjectPropertyEntryBox)
-		.AllowCreate(true)
+		.AllowCreate(false)
 		.AllowClear(true)
 		.DisplayThumbnail(false)
 		.DisplayCompactSize(true)
@@ -962,7 +962,7 @@ TSharedRef<ITableRow> SCineAssemblySchemaWindow::OnGenerateTreeRow(TSharedPtr<FS
 				.VAlign(VAlign_Center)
 				[
 					SAssignNew(TreeItem->TemplateObjectEntry, SObjectPropertyEntryBox)
-						.AllowCreate(true)
+						.AllowCreate(false)
 						.AllowClear(true)
 						.DisplayThumbnail(false)
 						.DisplayCompactSize(true)
