ProductionWizard: Fixing issue where inner structs of Extended Data Structs would not correctly update the config when changed.

#rnx

diff --git a/Engine/Plugins/VirtualProduction/CinematicAssemblyTools/Source/CineAssemblyToolsEditor/Private/ProductionSettings.cpp b/Engine/Plugins/VirtualProduction/CinematicAssemblyTools/Source/CineAssemblyToolsEditor/Private/ProductionSettings.cpp
--- a/Engine/Plugins/VirtualProduction/CinematicAssemblyTools/Source/CineAssemblyToolsEditor/Private/ProductionSettings.cpp
+++ b/Engine/Plugins/VirtualProduction/CinematicAssemblyTools/Source/CineAssemblyToolsEditor/Private/ProductionSettings.cpp
@@ -498,23 +498,14 @@ void UProductionSettings::PostEditChangeChainProperty(FPropertyChangedChainEvent
 {
 	Super::PostEditChangeChainProperty(PropertyChangedEvent);
 
-	// Is this a Production change?
-	if (TDoubleLinkedList<FProperty*>::TDoubleLinkedListNode* ProductionsNode = PropertyChangedEvent.PropertyChain.GetHead();
-		ProductionsNode != nullptr && ProductionsNode->GetValue()->GetFName() == GET_MEMBER_NAME_CHECKED(UProductionSettings, Productions))
+	if (TDoubleLinkedList<FProperty*>::TDoubleLinkedListNode* HeadNode = PropertyChangedEvent.PropertyChain.GetHead();
+		HeadNode != nullptr
+		&& HeadNode->GetValue() != nullptr)
 	{
-		// Check if this is is an ExtendedData change.
-		if (TDoubleLinkedList<FProperty*>::TDoubleLinkedListNode* ExtendedDataNode = ProductionsNode->GetNextNode();
-			ExtendedDataNode && ExtendedDataNode->GetValue()->GetFName() == FCinematicProduction::GetExtendedDataMemberName())
+		// Is this a Production change?
+		if (HeadNode->GetValue()->GetFName() == GET_MEMBER_NAME_CHECKED(UProductionSettings, Productions))
 		{
-			const int32 ProductionIndex = PropertyChangedEvent.GetArrayIndex(ProductionsNode->GetValue()->GetName());
-
-			if (Productions.IsValidIndex(ProductionIndex))
-			{
-				if (const UScriptStruct* OwnerStruct = PropertyChangedEvent.Property->GetTypedOwner<UScriptStruct>())
-				{
-					Productions[ProductionIndex].ExportExtendedData(*OwnerStruct);
-				}
-			}
+			HandleProductionsChangeChainProperty(*HeadNode, PropertyChangedEvent);
 		}
 	}
 }
@@ -629,6 +620,43 @@ void UProductionSettings::HandleProductionExtendedDataUpdated(const FCinematicPr
 	UpdateConfig();
 }
 
+void UProductionSettings::HandleProductionsChangeChainProperty(const TDoubleLinkedList<FProperty*>::TDoubleLinkedListNode& ProductionsChainNode, FPropertyChangedChainEvent& WithinPropertyChangedEvent)
+{
+	check(ProductionsChainNode.GetValue() != nullptr);
+
+	const int32 ProductionIndex = WithinPropertyChangedEvent.GetArrayIndex(ProductionsChainNode.GetValue()->GetName());
+
+	// If the index of the modified Production is not valid, return now as we can't handle anything
+	if (!Productions.IsValidIndex(ProductionIndex))
+	{
+		return;
+	}
+
+	// Check if this is is an ExtendedData change.
+	if (const TDoubleLinkedList<FProperty*>::TDoubleLinkedListNode* ExtendedDataNode = ProductionsChainNode.GetNextNode();
+		ExtendedDataNode != nullptr
+		&& ExtendedDataNode->GetValue() != nullptr
+		&& ExtendedDataNode->GetValue()->GetFName() == FCinematicProduction::GetExtendedDataMemberName())
+	{
+		// Get the extended data property node that changed. From here we can try to extract the type of ExtendedData.
+		const TDoubleLinkedList<FProperty*>::TDoubleLinkedListNode* ChangedDataPropertyNode = ExtendedDataNode->GetNextNode();
+
+		const UScriptStruct* ChangedStruct = ChangedDataPropertyNode != nullptr && ChangedDataPropertyNode->GetValue() != nullptr
+			? ChangedDataPropertyNode->GetValue()->GetTypedOwner<UScriptStruct>()
+			: nullptr;
+
+		if (ChangedStruct != nullptr)
+		{
+			Productions[ProductionIndex].ExportExtendedData(*ChangedStruct);
+		}
+		else
+		{
+			// Something else changed about the extended data but we can't get the exact type, maybe a reset, so fallback to an untargeted update of the production.
+			FScopedModifyProductionExtendedData ModifyGuard(Productions[ProductionIndex]);
+		}
+	}
+}
+
 void UProductionSettings::DeleteProduction(FGuid ProductionID)
 {
 	if (!ProductionID.IsValid())
diff --git a/Engine/Plugins/VirtualProduction/CinematicAssemblyTools/Source/CineAssemblyToolsEditor/Public/ProductionSettings.h b/Engine/Plugins/VirtualProduction/CinematicAssemblyTools/Source/CineAssemblyToolsEditor/Public/ProductionSettings.h
--- a/Engine/Plugins/VirtualProduction/CinematicAssemblyTools/Source/CineAssemblyToolsEditor/Public/ProductionSettings.h
+++ b/Engine/Plugins/VirtualProduction/CinematicAssemblyTools/Source/CineAssemblyToolsEditor/Public/ProductionSettings.h
@@ -450,6 +450,9 @@ private:
 	/** Handler for when a Productions ExtendedData is updated. */
 	void HandleProductionExtendedDataUpdated(const FCinematicProduction& SenderProduction, const UScriptStruct* ForData);
 
+	/** Given a property changed chain event and a node representing the productions property, handle any property changes that have occured to the Productions property. */
+	void HandleProductionsChangeChainProperty(const TDoubleLinkedList<FProperty*>::TDoubleLinkedListNode& ProductionsChainNode, FPropertyChangedChainEvent& WithinPropertyChangedEvent);
+
 private:
 	/** Name of the active production */
 	UPROPERTY(VisibleAnywhere, Category = "Default", meta = (DisplayName = "Active Production"))
