ShotManagement: Change Cine Assembly Schema asset template combo buttons to always be interactive and editable.
#rnx

diff --git a/Engine/Plugins/VirtualProduction/CinematicAssemblyTools/Source/CineAssemblyToolsEditor/Private/UI/CineAssembly/SCineAssemblySchemaWindow.cpp b/Engine/Plugins/VirtualProduction/CinematicAssemblyTools/Source/CineAssemblyToolsEditor/Private/UI/CineAssembly/SCineAssemblySchemaWindow.cpp
--- a/Engine/Plugins/VirtualProduction/CinematicAssemblyTools/Source/CineAssemblyToolsEditor/Private/UI/CineAssembly/SCineAssemblySchemaWindow.cpp
+++ b/Engine/Plugins/VirtualProduction/CinematicAssemblyTools/Source/CineAssemblyToolsEditor/Private/UI/CineAssembly/SCineAssemblySchemaWindow.cpp
@@ -14,6 +14,7 @@
 #include "Misc/MessageDialog.h"
 #include "Modules/ModuleManager.h"
 #include "NamingTokensEngineSubsystem.h"
+#include "PropertyCustomizationHelpers.h"
 #include "PropertyEditorDelegates.h"
 #include "PropertyEditorModule.h"
 #include "Settings/ContentBrowserSettings.h"
@@ -507,6 +508,7 @@ TSharedRef<SWidget> SCineAssemblySchemaWindow::MakeHierarchyTabContent()
 		.OnGetChildren(this, &SCineAssemblySchemaWindow::OnGetChildren)
 		.OnItemsRebuilt(this, &SCineAssemblySchemaWindow::OnTreeItemsRebuilt)
 		.OnContextMenuOpening(this, &SCineAssemblySchemaWindow::MakeContentTreeContextMenu)
+		.OnMouseButtonDoubleClick(this, &SCineAssemblySchemaWindow::OnTreeViewDoubleClick)
 		.OnKeyDownHandler(this, &SCineAssemblySchemaWindow::OnTreeViewKeyDown);
 
 	ExpandTreeRecursive(RootItem);
@@ -734,7 +736,7 @@ TSharedRef<ITableRow> SCineAssemblySchemaWindow::OnGenerateAssetRow(TSharedPtr<F
 		.ShowSelection(true)
 		.Padding(FMargin(4.0f, 2.0f))
 		.OnDragDetected(this, &SCineAssemblySchemaWindow::OnAssetRowDragDetected)
-		.ToolTipText(LOCTEXT("AssetRowTooltip", "Double click to rename the asset and change the template asset."));
+		.ToolTipText(LOCTEXT("AssetRowTooltip", "Set the name and template subsequence (optional) for this subsequence, then add it to the content tree on the right."));
 
 	// Store a reference to the editable textblock in the row to easily set it to edit mode for renaming
 	Row->TextBlock = SNew(SInlineEditableTextBlock)
@@ -748,38 +750,6 @@ TSharedRef<ITableRow> SCineAssemblySchemaWindow::OnGenerateAssetRow(TSharedPtr<F
 			InItem->Name = InText.ToString();
 		});
 
-	Row->TemplateObjectEntry = SNew(SObjectPropertyEntryBox)
-		.AllowCreate(false)
-		.AllowClear(true)
-		.DisplayThumbnail(false)
-		.DisplayCompactSize(true)
-		.AllowedClass(ULevelSequence::StaticClass())
-		.ToolTipText(LOCTEXT("SchemaAssetTemplateSelectionTooltip", "Select an existing Asset to use as a template for this asset."))
-		.Visibility_Lambda([WeakRow = TWeakPtr<SSchemaAssetTableRow>(Row)]() 
-			{ 
-				if (TSharedPtr<SSchemaAssetTableRow> RowPtr = WeakRow.Pin())
-				{
-					return RowPtr->TextBlock->IsInEditMode() ? EVisibility::Visible : EVisibility::HitTestInvisible;
-				}
-
-				return EVisibility::Hidden;
-			})
-		.ObjectPath_Lambda([WeakInItem = TWeakPtr<FSchemaAssetParameters>(InItem)]() -> FString
-			{
-				if (TSharedPtr<FSchemaAssetParameters> ItemPtr = WeakInItem.Pin())
-				{
-					return ItemPtr->Template.ToString();
-				}
-				return FString();
-			})
-		.OnObjectChanged_Lambda([WeakInItem = TWeakPtr<FSchemaAssetParameters>(InItem)](const FAssetData& InAssetData)
-			{
-				if (TSharedPtr<FSchemaAssetParameters> ItemPtr = WeakInItem.Pin())
-				{
-					ItemPtr->Template = InAssetData.GetSoftObjectPath();
-				}
-			});
-
 	Row->SetContent(SNew(SBorder)
 		.Padding(8.0f)
 		.BorderImage(FCineAssemblyToolsStyle::Get().GetBrush("ProductionWizard.PanelBackground"))
@@ -805,7 +775,28 @@ TSharedRef<ITableRow> SCineAssemblySchemaWindow::OnGenerateAssetRow(TSharedPtr<F
 				.FillWidth(1.0f)
 				.VAlign(VAlign_Center)
 				[
-					Row->TemplateObjectEntry.ToSharedRef()
+					SNew(SObjectPropertyEntryBox)
+						.AllowCreate(false)
+						.AllowClear(true)
+						.DisplayThumbnail(false)
+						.DisplayCompactSize(true)
+						.AllowedClass(ULevelSequence::StaticClass())
+						.ToolTipText(LOCTEXT("SchemaAssetTemplateSelectionTooltip", "(Optional) Select an existing level sequence to use as a template for this subsequence."))
+						.ObjectPath_Lambda([WeakInItem = TWeakPtr<FSchemaAssetParameters>(InItem)]() -> FString
+						{
+							if (TSharedPtr<FSchemaAssetParameters> ItemPtr = WeakInItem.Pin())
+							{
+								return ItemPtr->Template.ToString();
+							}
+							return FString();
+						})
+						.OnObjectChanged_Lambda([WeakInItem = TWeakPtr<FSchemaAssetParameters>(InItem)](const FAssetData& InAssetData)
+						{
+							if (TSharedPtr<FSchemaAssetParameters> ItemPtr = WeakInItem.Pin())
+							{
+								ItemPtr->Template = InAssetData.GetSoftObjectPath();
+							}
+						})
 				]
 
 			+ SHorizontalBox::Slot()
@@ -909,20 +900,6 @@ TSharedRef<ITableRow> SCineAssemblySchemaWindow::OnGenerateTreeRow(TSharedPtr<FS
 			})
 		.OnAcceptDrop(this, &SCineAssemblySchemaWindow::OnTreeRowAcceptDrop)
 		.OnDragDetected(this, &SCineAssemblySchemaWindow::OnTreeRowDragDetected)
-		.ToolTipText_Lambda([this, WeakTreeItem = TWeakPtr<FSchemaTreeItem>(TreeItem)]() -> FText
-			{
-				if (TSharedPtr<FSchemaTreeItem> ItemPtr = WeakTreeItem.Pin())
-				{
-					if (ItemPtr->Type == FSchemaTreeItem::EItemType::Asset)
-					{
-						return ItemPtr == TopLevelAssemblyItem
-							? LOCTEXT("AssetTreeTopLevelAssemblyTooltip", "Right click to change the Assemblies template asset.")
-							: LOCTEXT("AssetTreeSubAssemblyTooltip", "Right click to rename the SubAssembly or change the template asset.");
-					}
-				}
-
-				return FText::GetEmpty();
-			})
 		.Content()
 		[
 			SNew(SHorizontalBox)
@@ -961,14 +938,14 @@ TSharedRef<ITableRow> SCineAssemblySchemaWindow::OnGenerateTreeRow(TSharedPtr<FS
 			+ SHorizontalBox::Slot()
 				.VAlign(VAlign_Center)
 				[
-					SAssignNew(TreeItem->TemplateObjectEntry, SObjectPropertyEntryBox)
+					SNew(SObjectPropertyEntryBox)
 						.AllowCreate(false)
 						.AllowClear(true)
 						.DisplayThumbnail(false)
 						.DisplayCompactSize(true)
 						.AllowedClass(ULevelSequence::StaticClass())
-						.ToolTipText(LOCTEXT("SchemaAssetTemplateSelectionTooltip", "Select an existing Asset to use as a template for this asset."))
-						.Visibility(TreeItem->Type == FSchemaTreeItem::EItemType::Asset ? EVisibility::HitTestInvisible : EVisibility::Collapsed)
+						.ToolTipText(LOCTEXT("SchemaAssetTemplateSelectionTooltip", "(Optional) Select an existing level sequence to use as a template for this subsequence."))
+						.Visibility(TreeItem->Type == FSchemaTreeItem::EItemType::Asset ? EVisibility::Visible : EVisibility::Collapsed)
 						.ObjectPath_Lambda([WeakTreeItem = TWeakPtr<FSchemaTreeItem>(TreeItem)]() -> FString
 							{
 								if (TSharedPtr<FSchemaTreeItem> ItemPtr = WeakTreeItem.Pin())
@@ -1023,54 +1000,30 @@ TSharedPtr<SWidget> SCineAssemblySchemaWindow::MakeContentTreeContextMenu()
 			);
 		}
 
-		if (SelectedTreeItem != RootItem)
+		if ((SelectedTreeItem != RootItem) && (SelectedTreeItem != TopLevelAssemblyItem)) 
 		{
-			if (SelectedTreeItem->Type == FSchemaTreeItem::EItemType::Asset)
-			{
-				MenuBuilder.AddMenuEntry(
-					LOCTEXT("EditTemplateAction", "Edit Template"),
-					LOCTEXT("EditTemplateActionToolTip", "Edit the assets template"),
-					FSlateIcon(FAppStyle::GetAppStyleSetName(), "Icons.Edit"),
-					FUIAction(
-						FExecuteAction::CreateLambda([SelectedTreeItem]()
-							{
-								if (SelectedTreeItem->TemplateObjectEntry.IsValid())
-								{
-									FSlateApplication::Get().SetKeyboardFocus(SelectedTreeItem->TemplateObjectEntry.ToSharedRef());
-									SelectedTreeItem->TemplateObjectEntry->SetVisibility(EVisibility::Visible);
-									SelectedTreeItem->TemplateObjectEntry->OpenEntryBox();
-								}
-							})),
-					NAME_None,
-					EUserInterfaceActionType::Button
-				);
-			}
+			MenuBuilder.AddMenuEntry(
+				LOCTEXT("RenameAction", "Rename"),
+				LOCTEXT("RenameActionToolTip", "Rename"),
+				FSlateIcon(FCineAssemblyToolsStyle::StyleName, "Icons.AssetNaming"),
+				FUIAction(
+					FExecuteAction::CreateLambda([SelectedTreeItem]()
+						{
+							FSlateApplication::Get().SetKeyboardFocus(SelectedTreeItem->NameWidget.ToSharedRef());
+							SelectedTreeItem->NameWidget->EnterEditingMode();
+						})),
+				NAME_None,
+				EUserInterfaceActionType::Button
+			);
 
-			if (SelectedTreeItem != TopLevelAssemblyItem)
-			{
-				MenuBuilder.AddMenuEntry(
-					LOCTEXT("RenameAction", "Rename"),
-					LOCTEXT("RenameActionToolTip", "Rename"),
-					FSlateIcon(FCineAssemblyToolsStyle::StyleName, "Icons.AssetNaming"),
-					FUIAction(
-						FExecuteAction::CreateLambda([SelectedTreeItem]()
-							{
-								FSlateApplication::Get().SetKeyboardFocus(SelectedTreeItem->NameWidget.ToSharedRef());
-								SelectedTreeItem->NameWidget->EnterEditingMode();
-							})),
-					NAME_None,
-					EUserInterfaceActionType::Button
-				);
-
-				MenuBuilder.AddMenuEntry(
-					LOCTEXT("DeleteAction", "Delete"),
-					LOCTEXT("DeleteActionToolTip", "Delete"),
-					FSlateIcon(FAppStyle::GetAppStyleSetName(), "Icons.Delete"),
-					FUIAction(FExecuteAction::CreateSP(this, &SCineAssemblySchemaWindow::DeleteTreeItem, SelectedTreeItem)),
-					NAME_None,
-					EUserInterfaceActionType::Button
-				);
-			}
+			MenuBuilder.AddMenuEntry(
+				LOCTEXT("DeleteAction", "Delete"),
+				LOCTEXT("DeleteActionToolTip", "Delete"),
+				FSlateIcon(FAppStyle::GetAppStyleSetName(), "Icons.Delete"),
+				FUIAction(FExecuteAction::CreateSP(this, &SCineAssemblySchemaWindow::DeleteTreeItem, SelectedTreeItem)),
+				NAME_None,
+				EUserInterfaceActionType::Button
+			);
 		}
 
 		return MenuBuilder.MakeWidget();
@@ -1079,6 +1032,19 @@ TSharedPtr<SWidget> SCineAssemblySchemaWindow::MakeContentTreeContextMenu()
 	return SNullWidget::NullWidget;
 }
 
+void SCineAssemblySchemaWindow::OnTreeViewDoubleClick(TSharedPtr<FSchemaTreeItem> TreeItem)
+{
+	// Get the row for the input item and put its textblock in edit mode so the user can rename the item
+	if (TreeItem && TreeItem->NameWidget)
+	{
+		if (TreeItem != TopLevelAssemblyItem)
+		{
+			FSlateApplication::Get().SetKeyboardFocus(TreeItem->NameWidget.ToSharedRef());
+			TreeItem->NameWidget->EnterEditingMode();
+		}
+	}
+}
+
 FReply SCineAssemblySchemaWindow::OnTreeViewKeyDown(const FGeometry& MyGeometry, const FKeyEvent& InKeyEvent)
 {
 	if (InKeyEvent.GetKey() == EKeys::Delete)
@@ -1295,7 +1261,6 @@ void SCineAssemblySchemaWindow::OnTreeItemTemplateChanged(const FAssetData& InAs
 	}
 
 	TreeItem->AssetTemplate = NewSoftPath;
-	TreeItem->TemplateObjectEntry->SetVisibility(EVisibility::HitTestInvisible);
 	UpdateAssetList();
 	TreeView->RequestTreeRefresh();
 }
diff --git a/Engine/Plugins/VirtualProduction/CinematicAssemblyTools/Source/CineAssemblyToolsEditor/Private/UI/CineAssembly/SCineAssemblySchemaWindow.h b/Engine/Plugins/VirtualProduction/CinematicAssemblyTools/Source/CineAssemblyToolsEditor/Private/UI/CineAssembly/SCineAssemblySchemaWindow.h
--- a/Engine/Plugins/VirtualProduction/CinematicAssemblyTools/Source/CineAssemblyToolsEditor/Private/UI/CineAssembly/SCineAssemblySchemaWindow.h
+++ b/Engine/Plugins/VirtualProduction/CinematicAssemblyTools/Source/CineAssemblyToolsEditor/Private/UI/CineAssembly/SCineAssemblySchemaWindow.h
@@ -6,7 +6,6 @@
 
 #include "CineAssemblySchema.h"
 #include "Input/DragAndDrop.h"
-#include "PropertyCustomizationHelpers.h"
 #include "Widgets/Layout/SWidgetSwitcher.h"
 #include "Widgets/Text/SInlineEditableTextBlock.h"
 #include "Widgets/Views/SListView.h"
@@ -37,9 +36,6 @@ class SSchemaAssetTableRow : public STableRow<TSharedPtr<FSchemaAssetParameters>
 public:
 	/** Editable text widget stored in this row to easily trigger edit mode */
 	TSharedPtr<SInlineEditableTextBlock> TextBlock;
-
-	/** Object entry box for picking possible Template sequences, and stored in this row to easily trigger edit mode. */
-	TSharedPtr<SObjectPropertyEntryBox> TemplateObjectEntry;
 };
 
 /** An entry in the content folder tree view with knowledge of its path, children, and parent */
@@ -72,9 +68,6 @@ struct FSchemaTreeItem
 
 	/** The text widget that displays this item's name and supports renaming */
 	TSharedPtr<SInlineEditableTextBlock> NameWidget;
-
-	/** Object entry box for picking possible Template sequences and supports changing the template. */
-	TSharedPtr<SObjectPropertyEntryBox> TemplateObjectEntry;
 };
 
 /** Drag and Drop operation to handle dragging schema assets and folders  */
@@ -210,6 +203,9 @@ private:
 	/** Create the context menu when the content tree view is right-clicked */
 	TSharedPtr<SWidget> MakeContentTreeContextMenu();
 
+	/** Handles double clicks on the tree view to rename the item */
+	void OnTreeViewDoubleClick(TSharedPtr<FSchemaTreeItem> TreeItem);
+
 	/** Handles key presses on the tree view */
 	FReply OnTreeViewKeyDown(const FGeometry& MyGeometry, const FKeyEvent& InKeyEvent);
 
