ShotManagement: Fix CIne Assemblies and Schema to properly notify the asset registry of their creation so that they will properly transact over multi-user.

diff --git a/Engine/Plugins/VirtualProduction/CinematicAssemblyTools/Source/CineAssemblyToolsEditor/Private/CineAssemblyFactory.cpp b/Engine/Plugins/VirtualProduction/CinematicAssemblyTools/Source/CineAssemblyToolsEditor/Private/CineAssemblyFactory.cpp
--- a/Engine/Plugins/VirtualProduction/CinematicAssemblyTools/Source/CineAssemblyToolsEditor/Private/CineAssemblyFactory.cpp
+++ b/Engine/Plugins/VirtualProduction/CinematicAssemblyTools/Source/CineAssemblyToolsEditor/Private/CineAssemblyFactory.cpp
@@ -2,6 +2,7 @@
 
 #include "CineAssemblyFactory.h"
 
+#include "AssetRegistry/AssetRegistryModule.h"
 #include "AssetViewUtils.h"
 #include "AssetToolsModule.h"
 #include "CineAssemblyNamingTokens.h"
@@ -177,6 +178,12 @@ void UCineAssemblyFactory::CreateConfiguredAssembly(UCineAssembly* ConfiguredAss
 	// Analytics
 	UE::CineAssemblyToolsAnalytics::RecordEvent_CreateAssembly(ConfiguredAssembly);
 
+	// Notify the asset registry about the new assembly
+	FAssetRegistryModule::AssetCreated(ConfiguredAssembly);
+
+	// Mark the package dirty
+	Package->MarkPackageDirty();
+
 	// Refresh the content browser to make any new assets and folders are immediately visible
 	IContentBrowserSingleton& ContentBrowser = FModuleManager::LoadModuleChecked<FContentBrowserModule>("ContentBrowser").Get();
 	ContentBrowser.SetSelectedPaths({ CreateAssetPath }, true);
diff --git a/Engine/Plugins/VirtualProduction/CinematicAssemblyTools/Source/CineAssemblyToolsEditor/Private/CineAssemblySchemaFactory.cpp b/Engine/Plugins/VirtualProduction/CinematicAssemblyTools/Source/CineAssemblyToolsEditor/Private/CineAssemblySchemaFactory.cpp
--- a/Engine/Plugins/VirtualProduction/CinematicAssemblyTools/Source/CineAssemblyToolsEditor/Private/CineAssemblySchemaFactory.cpp
+++ b/Engine/Plugins/VirtualProduction/CinematicAssemblyTools/Source/CineAssemblyToolsEditor/Private/CineAssemblySchemaFactory.cpp
@@ -2,6 +2,7 @@
 
 #include "CineAssemblySchemaFactory.h"
 
+#include "AssetRegistry/AssetRegistryModule.h"
 #include "AssetViewUtils.h"
 #include "AssetToolsModule.h"
 #include "CineAssemblySchema.h"
@@ -100,6 +101,12 @@ void UCineAssemblySchemaFactory::CreateConfiguredSchema(UCineAssemblySchema* Con
 
 	UE::CineAssemblyToolsAnalytics::RecordEvent_CreateAssemblySchema();
 
+	// Notify the asset registry about the new schema
+	FAssetRegistryModule::AssetCreated(ConfiguredSchema);
+
+	// Mark the package dirty
+	Package->MarkPackageDirty();
+
 	// Refresh the content browser to make any new assets and folders are immediately visible
 	IContentBrowserSingleton& ContentBrowser = FModuleManager::LoadModuleChecked<FContentBrowserModule>("ContentBrowser").Get();
 	ContentBrowser.SetSelectedPaths({ CreateAssetPath }, true);
