ShotManagement: Fix schema name property "reset to default" functionality.
#rnx

diff --git a/Engine/Plugins/VirtualProduction/CinematicAssemblyTools/Source/CineAssemblyToolsEditor/Private/CineAssemblySchemaCustomization.cpp b/Engine/Plugins/VirtualProduction/CinematicAssemblyTools/Source/CineAssemblyToolsEditor/Private/CineAssemblySchemaCustomization.cpp
--- a/Engine/Plugins/VirtualProduction/CinematicAssemblyTools/Source/CineAssemblyToolsEditor/Private/CineAssemblySchemaCustomization.cpp
+++ b/Engine/Plugins/VirtualProduction/CinematicAssemblyTools/Source/CineAssemblyToolsEditor/Private/CineAssemblySchemaCustomization.cpp
@@ -41,13 +41,13 @@ void FCineAssemblySchemaCustomization::CustomizeDetails(IDetailLayoutBuilder& De
 	IDetailPropertyRow& DefaultAssemblyNamePropertyRow = DefaultCategory.AddProperty(DefaultAssemblyNamePropertyHandle);
 
 	// Customize the widget for the SchemaName property to add additional validation on the user input text
-	FString ExistingSchemaName;
-	SchemaNamePropertyHandle->GetValue(ExistingSchemaName);
+	SchemaNamePropertyHandle->GetValue(OriginalSchemaName);
 
 	// If the key name is not yet set, assign it a unique default key name
-	if (ExistingSchemaName.IsEmpty())
+	if (OriginalSchemaName.IsEmpty())
 	{
-		SchemaNamePropertyHandle->SetValue(MakeUniqueSchemaName());
+		OriginalSchemaName = MakeUniqueSchemaName();
+		SchemaNamePropertyHandle->SetValue(OriginalSchemaName);
 	}
 
 	SchemaNamePropertyRow.CustomWidget()
@@ -70,6 +70,11 @@ void FCineAssemblySchemaCustomization::CustomizeDetails(IDetailLayoutBuilder& De
 				.OnVerifyTextChanged_Raw(this, &FCineAssemblySchemaCustomization::ValidateSchemaName)
 		];
 
+	SchemaNamePropertyRow.OverrideResetToDefault(
+		FResetToDefaultOverride::Create(
+			FIsResetToDefaultVisible::CreateSP(this, &FCineAssemblySchemaCustomization::IsSchemaNameResetToDefaultVisible),
+			FResetToDefaultHandler::CreateSP(this, &FCineAssemblySchemaCustomization::OnSchemaNameResetToDefault)));
+
 	// Customize the widget for the DefaultAssemblyName property to add additional validation on the user input text
 	DefaultAssemblyNamePropertyRow.CustomWidget()
 		.NameContent()
@@ -179,4 +184,20 @@ bool FCineAssemblySchemaCustomization::ValidateDefaultAssemblyName(const FText&
 	return true;
 }
 
+bool FCineAssemblySchemaCustomization::IsSchemaNameResetToDefaultVisible(TSharedPtr<IPropertyHandle> SchemaNamePropertyHandle) const
+{
+	// The "default" name of the schema is defined as follows:
+	// For new schemas, this will be the default (unique) name that was assigned to the schema
+	// For existing schemas, this will be the name of the schema when its asset editor window was opened
+	FString CurrentSchemaName;
+	SchemaNamePropertyHandle->GetValue(CurrentSchemaName);
+
+	return !CurrentSchemaName.Equals(OriginalSchemaName);
+};
+
+void FCineAssemblySchemaCustomization::OnSchemaNameResetToDefault(TSharedPtr<IPropertyHandle> SchemaNamePropertyHandle)
+{
+	CustomizedSchema->RenameAsset(OriginalSchemaName);
+}
+
 #undef LOCTEXT_NAMESPACE
diff --git a/Engine/Plugins/VirtualProduction/CinematicAssemblyTools/Source/CineAssemblyToolsEditor/Private/CineAssemblySchemaCustomization.h b/Engine/Plugins/VirtualProduction/CinematicAssemblyTools/Source/CineAssemblyToolsEditor/Private/CineAssemblySchemaCustomization.h
--- a/Engine/Plugins/VirtualProduction/CinematicAssemblyTools/Source/CineAssemblyToolsEditor/Private/CineAssemblySchemaCustomization.h
+++ b/Engine/Plugins/VirtualProduction/CinematicAssemblyTools/Source/CineAssemblyToolsEditor/Private/CineAssemblySchemaCustomization.h
@@ -5,6 +5,7 @@
 #include "IDetailCustomization.h"
 
 class UCineAssemblySchema;
+class IPropertyHandle;
 
 /**
  * Detail customization for UCineAssemblySchema
@@ -30,7 +31,14 @@ private:
 	bool ValidateSchemaName(const FText& InText, FText& OutErrorMessage) const;
 	bool ValidateDefaultAssemblyName(const FText& InText, FText& OutErrorMessage) const;
 
+	/** Implements "reset to default" behavior for the schema name property */
+	bool IsSchemaNameResetToDefaultVisible(TSharedPtr<IPropertyHandle> SchemaNamePropertyHandle) const;
+	void OnSchemaNameResetToDefault(TSharedPtr<IPropertyHandle> SchemaNamePropertyHandle);
+
 private:
 	/** The assembly schema being customized */
 	UCineAssemblySchema* CustomizedSchema;
+
+	/** Original name of the schema being customized */
+	FString OriginalSchemaName;
 };
