ShotManagement: Fix incorrect delegate binding
#rnx

diff --git a/Engine/Plugins/VirtualProduction/CinematicAssemblyTools/Source/CineAssemblyToolsEditor/Private/CineAssemblySchemaCustomization.cpp b/Engine/Plugins/VirtualProduction/CinematicAssemblyTools/Source/CineAssemblyToolsEditor/Private/CineAssemblySchemaCustomization.cpp
--- a/Engine/Plugins/VirtualProduction/CinematicAssemblyTools/Source/CineAssemblyToolsEditor/Private/CineAssemblySchemaCustomization.cpp
+++ b/Engine/Plugins/VirtualProduction/CinematicAssemblyTools/Source/CineAssemblyToolsEditor/Private/CineAssemblySchemaCustomization.cpp
@@ -72,8 +72,8 @@ void FCineAssemblySchemaCustomization::CustomizeDetails(IDetailLayoutBuilder& De
 
 	SchemaNamePropertyRow.OverrideResetToDefault(
 		FResetToDefaultOverride::Create(
-			FIsResetToDefaultVisible::CreateSP(this, &FCineAssemblySchemaCustomization::IsSchemaNameResetToDefaultVisible),
-			FResetToDefaultHandler::CreateSP(this, &FCineAssemblySchemaCustomization::OnSchemaNameResetToDefault)));
+			FIsResetToDefaultVisible::CreateRaw(this, &FCineAssemblySchemaCustomization::IsSchemaNameResetToDefaultVisible),
+			FResetToDefaultHandler::CreateRaw(this, &FCineAssemblySchemaCustomization::OnSchemaNameResetToDefault)));
 
 	// Customize the widget for the DefaultAssemblyName property to add additional validation on the user input text
 	DefaultAssemblyNamePropertyRow.CustomWidget()
