ShotManagement: Fix CAT metadata tokens from triggering recursive token evaluation.
#rnx

#FailedValidations Preflight

diff --git a/Engine/Plugins/VirtualProduction/CinematicAssemblyTools/Source/CineAssemblyTools/Private/CineAssemblyNamingTokens.cpp b/Engine/Plugins/VirtualProduction/CinematicAssemblyTools/Source/CineAssemblyTools/Private/CineAssemblyNamingTokens.cpp
--- a/Engine/Plugins/VirtualProduction/CinematicAssemblyTools/Source/CineAssemblyTools/Private/CineAssemblyNamingTokens.cpp
+++ b/Engine/Plugins/VirtualProduction/CinematicAssemblyTools/Source/CineAssemblyTools/Private/CineAssemblyNamingTokens.cpp
@@ -213,7 +213,7 @@ void UCineAssemblyNamingTokens::AddMetadataToken(const FString& InTokenKey)
 			if (InAssembly.IsValid())
 			{
 				FString ValueString;
-				if (InAssembly->GetMetadataAsString(InTokenKey, ValueString))
+				if (InAssembly->GetMetadataAsString(InTokenKey, ValueString, false))
 				{
 					// Test if the value string could be an object path, and if it is, return just the filename instead of the full path
 					FSoftObjectPath ObjectPath = FSoftObjectPath(ValueString);
