ShotManagement: Move the ownership of Schema thumbnail Slate brushes out of the Schema assets and into the CAT style set. This fixes a bug where the editor could crash after deleting a Schema asset because the brush could become invalid while Slate was still using it.
#rnx
#FailedValidations Preflight

diff --git a/Engine/Plugins/VirtualProduction/CinematicAssemblyTools/Source/CineAssemblyTools/Private/CineAssemblySchema.cpp b/Engine/Plugins/VirtualProduction/CinematicAssemblyTools/Source/CineAssemblyTools/Private/CineAssemblySchema.cpp
--- a/Engine/Plugins/VirtualProduction/CinematicAssemblyTools/Source/CineAssemblyTools/Private/CineAssemblySchema.cpp
+++ b/Engine/Plugins/VirtualProduction/CinematicAssemblyTools/Source/CineAssemblyTools/Private/CineAssemblySchema.cpp
@@ -10,17 +10,10 @@
 #include "AssetToolsModule.h"
 #endif
 
-const FSoftObjectPath UCineAssemblySchema::DefaultThumbnailPath = FSoftObjectPath(TEXT("/CinematicAssemblyTools/Resources/DefaultSchemaThumbnail.DefaultSchemaThumbnail"));
 const FName UCineAssemblySchema::SchemaGuidPropertyName = GET_MEMBER_NAME_CHECKED(UCineAssemblySchema, SchemaGuid);
 
 UCineAssemblySchema::UCineAssemblySchema()
 {
-	if (!HasAnyFlags(RF_ArchetypeObject | RF_ClassDefaultObject))
-	{
-		ThumbnailBrush = MakeShared<FSlateBrush>();
-		ThumbnailBrush->ImageSize = FVector2D(64.0f, 64.0f);
-		UpdateThumbnailBrush();
-	}
 }
 
 FGuid UCineAssemblySchema::GetSchemaGuid() const
@@ -65,25 +58,6 @@ void UCineAssemblySchema::PostLoad()
 	{
 		SchemaGuid = FGuid::NewGuid();
 	}
-
-	UpdateThumbnailBrush();
-}
-
-#if WITH_EDITOR
-void UCineAssemblySchema::PostEditChangeProperty(struct FPropertyChangedEvent& PropertyChangedEvent)
-{
-	const FName PropertyName = PropertyChangedEvent.GetPropertyName();
-
-	if (PropertyName == GET_MEMBER_NAME_CHECKED(UCineAssemblySchema, ThumbnailImage))
-	{
-		UpdateThumbnailBrush();
-	}
-}
-#endif
-
-const FSlateBrush* UCineAssemblySchema::GetThumbnailBrush() const
-{
-	return ThumbnailBrush.Get();
 }
 
 bool UCineAssemblySchema::SupportsRename() const
@@ -129,18 +103,3 @@ void UCineAssemblySchema::RenameAsset(const FString& InNewName)
 	bSupportsRename = false;
 }
 #endif
-
-void UCineAssemblySchema::UpdateThumbnailBrush()
-{
-	if (ThumbnailImage)
-	{
-		ThumbnailBrush->SetResourceObject(ThumbnailImage);
-	}
-	else
-	{
-		FAssetRegistryModule& AssetRegistryModule = FModuleManager::GetModuleChecked<FAssetRegistryModule>("AssetRegistry");
-		const FAssetData ThumbnailAssetData = AssetRegistryModule.Get().GetAssetByObjectPath(DefaultThumbnailPath);
-
-		ThumbnailBrush->SetResourceObject(ThumbnailAssetData.GetAsset());
-	}
-}
diff --git a/Engine/Plugins/VirtualProduction/CinematicAssemblyTools/Source/CineAssemblyTools/Public/CineAssemblySchema.h b/Engine/Plugins/VirtualProduction/CinematicAssemblyTools/Source/CineAssemblyTools/Public/CineAssemblySchema.h
--- a/Engine/Plugins/VirtualProduction/CinematicAssemblyTools/Source/CineAssemblyTools/Public/CineAssemblySchema.h
+++ b/Engine/Plugins/VirtualProduction/CinematicAssemblyTools/Source/CineAssemblyTools/Public/CineAssemblySchema.h
@@ -65,18 +65,11 @@ public:
 	virtual void PostInitProperties() override;
 	virtual void PostDuplicate(bool bDuplicateForPIE) override;
 	virtual void PostLoad() override;
-
-#if WITH_EDITOR
-	virtual void PostEditChangeProperty(struct FPropertyChangedEvent& PropertyChangedEvent) override;
-#endif
 	//~ End UObject overrides
 
 	/** Get the unique ID of this schema */
 	CINEASSEMBLYTOOLS_API FGuid GetSchemaGuid() const;
 
-	/** Returns the thumbnail brush used by this schema */
-	CINEASSEMBLYTOOLS_API const FSlateBrush* GetThumbnailBrush() const;
-
 	/** Whether or not this schema can be renamed */
 	CINEASSEMBLYTOOLS_API bool SupportsRename() const;
 
@@ -85,10 +78,6 @@ public:
 	CINEASSEMBLYTOOLS_API void RenameAsset(const FString& InNewName);
 #endif
 
-private:
-	/** Updates the thumbnail brush resource */
-	void UpdateThumbnailBrush();
-
 public:
 	/** The schema name, which will be used by assemblies made from this schema as their "assembly type" */
 	UPROPERTY(EditAnywhere, Category = "Default")
@@ -148,15 +137,6 @@ private:
 	UPROPERTY(AssetRegistrySearchable, meta = (IgnoreForMemberInitializationTest))
 	FGuid SchemaGuid;
 
-	/** Slate brush used to actually render the thumbnail image */
-	TSharedPtr<FSlateBrush> ThumbnailBrush;
-
-	/** Whether the thumbnail brush needs to be updated */
-	bool bThumbnailPendingReset = false;
-
 	/** Whether the schema asset supports renaming */
 	bool bSupportsRename = false;
-
-	/** Asset path of the default thumbnail to use for schema assets */
-	static const FSoftObjectPath DefaultThumbnailPath;
 };
diff --git a/Engine/Plugins/VirtualProduction/CinematicAssemblyTools/Source/CineAssemblyToolsEditor/Private/AssetDefinition_CineAssembly.cpp b/Engine/Plugins/VirtualProduction/CinematicAssemblyTools/Source/CineAssemblyToolsEditor/Private/AssetDefinition_CineAssembly.cpp
--- a/Engine/Plugins/VirtualProduction/CinematicAssemblyTools/Source/CineAssemblyToolsEditor/Private/AssetDefinition_CineAssembly.cpp
+++ b/Engine/Plugins/VirtualProduction/CinematicAssemblyTools/Source/CineAssemblyToolsEditor/Private/AssetDefinition_CineAssembly.cpp
@@ -1,6 +1,7 @@
 // Copyright Epic Games, Inc. All Rights Reserved.
 
 #include "AssetDefinition_CineAssembly.h"
+#include "AssetToolsModule.h"
 #include "CineAssembly.h"
 #include "CineAssemblyToolsStyle.h"
 
@@ -61,11 +62,19 @@ const FSlateBrush* UAssetDefinition_CineAssembly::GetThumbnailBrush(const FAsset
 {
 	if (const UCineAssembly* const CineAssembly = Cast<UCineAssembly>(InAssetData.GetAsset()))
 	{
-		if (const UCineAssemblySchema* Schema = CineAssembly->GetSchema())
+		const UCineAssemblySchema* Schema = CineAssembly->GetSchema();
+		if (Schema && Schema->ThumbnailImage)
 		{
-			if (Schema->ThumbnailImage)
+			// Use the thumbnail brush associated with the schema
+			FAssetToolsModule& AssetToolsModule = FModuleManager::LoadModuleChecked<FAssetToolsModule>(TEXT("AssetTools"));
+			TSharedPtr<IAssetTypeActions> AssetTypeActions = AssetToolsModule.Get().GetAssetTypeActionsForClass(UCineAssemblySchema::StaticClass()).Pin();
+
+			if (AssetTypeActions)
 			{
-				return Schema->GetThumbnailBrush();
+				const FAssetData AssetData = FAssetData(Schema);
+				const FName AssetClassName = AssetData.AssetClassPath.GetAssetName();
+
+				return AssetTypeActions->GetThumbnailBrush(AssetData, AssetClassName);
 			}
 		}
 	}
diff --git a/Engine/Plugins/VirtualProduction/CinematicAssemblyTools/Source/CineAssemblyToolsEditor/Private/AssetDefinition_CineAssemblySchema.cpp b/Engine/Plugins/VirtualProduction/CinematicAssemblyTools/Source/CineAssemblyToolsEditor/Private/AssetDefinition_CineAssemblySchema.cpp
--- a/Engine/Plugins/VirtualProduction/CinematicAssemblyTools/Source/CineAssemblyToolsEditor/Private/AssetDefinition_CineAssemblySchema.cpp
+++ b/Engine/Plugins/VirtualProduction/CinematicAssemblyTools/Source/CineAssemblyToolsEditor/Private/AssetDefinition_CineAssemblySchema.cpp
@@ -34,10 +34,20 @@ const FSlateBrush* UAssetDefinition_CineAssemblySchema::GetThumbnailBrush(const
 {
 	if (UCineAssemblySchema* CineAssemblySchema = Cast<UCineAssemblySchema>(InAssetData.GetAsset()))
 	{
-		if (CineAssemblySchema->ThumbnailImage)
+		FCineAssemblyToolsStyle& Style = FCineAssemblyToolsStyle::Get();
+
+		// Get the thumbnail brush associated with this schema asset
+		const FName ThumbnailBrushName = Style.GetThumbnailBrushNameForSchema(CineAssemblySchema->SchemaName);
+		const FSlateBrush* ThumbnailBrush = Style.GetBrush(ThumbnailBrushName);
+
+		// If no thumbnail brush exists yet for this schema, create a new one
+		if (!ThumbnailBrush || ThumbnailBrush == Style.GetDefaultBrush())
 		{
-			return CineAssemblySchema->GetThumbnailBrush();
+			Style.SetThumbnailBrushTextureForSchema(CineAssemblySchema->SchemaName, CineAssemblySchema->ThumbnailImage);
+			ThumbnailBrush = Style.GetBrush(ThumbnailBrushName);
 		}
+
+		return ThumbnailBrush;
 	}
 
 	return FCineAssemblyToolsStyle::Get().GetBrush("ClassThumbnail.CineAssemblySchema");
diff --git a/Engine/Plugins/VirtualProduction/CinematicAssemblyTools/Source/CineAssemblyToolsEditor/Private/CineAssemblyToolsStyle.cpp b/Engine/Plugins/VirtualProduction/CinematicAssemblyTools/Source/CineAssemblyToolsEditor/Private/CineAssemblyToolsStyle.cpp
--- a/Engine/Plugins/VirtualProduction/CinematicAssemblyTools/Source/CineAssemblyToolsEditor/Private/CineAssemblyToolsStyle.cpp
+++ b/Engine/Plugins/VirtualProduction/CinematicAssemblyTools/Source/CineAssemblyToolsEditor/Private/CineAssemblyToolsStyle.cpp
@@ -2,6 +2,7 @@
 
 #include "CineAssemblyToolsStyle.h"
 
+#include "Engine/Texture2D.h"
 #include "Interfaces/IPluginManager.h"
 #include "Styling/CoreStyle.h"
 #include "Styling/SegmentedControlStyle.h"
@@ -118,3 +119,41 @@ FCineAssemblyToolsStyle::~FCineAssemblyToolsStyle()
 {
 	FSlateStyleRegistry::UnRegisterSlateStyle(*this);
 }
+
+FName FCineAssemblyToolsStyle::GetThumbnailBrushNameForSchema(const FString& SchemaName) const
+{
+	const FName BrushName = *FString::Printf(TEXT("SchemaThumbnail.%s"), *SchemaName);
+	return BrushName;
+}
+
+void FCineAssemblyToolsStyle::SetThumbnailBrushTextureForSchema(const FString& SchemaName, UTexture2D* BrushTexture)
+{
+	if (FSlateBrush* Result = BrushResources.FindRef(GetThumbnailBrushNameForSchema(SchemaName)))
+	{
+		// Update the resource object and image type of the existing brush
+		if (BrushTexture)
+		{
+			Result->SetResourceObject(BrushTexture);
+			Result->ImageType = ESlateBrushImageType::FullColor;
+		}
+		else
+		{
+			Result->SetResourceObject(nullptr);
+			Result->ImageType = ESlateBrushImageType::Vector;
+		}
+	}
+	else
+	{
+		// The new brush is initially created to use the schema icon so that it has a valid ResourceName, because the ResourceName property cannot be changed after the brush is created.
+		// If the brush's ResourceObject is ever null, a valid ResourceName allows it to immediately fallback to using the schema icon instead.
+		FSlateBrush* NewBrush = new FSlateVectorImageBrush(RootToContentDir(TEXT("Slate/SMAssemblySchema_64.svg")), FVector2D(64.0f, 64.0f));
+
+		if (BrushTexture)
+		{
+			NewBrush->SetResourceObject(BrushTexture);
+			NewBrush->ImageType = ESlateBrushImageType::FullColor;
+		}
+
+		Set(GetThumbnailBrushNameForSchema(SchemaName), NewBrush);
+	}
+}
diff --git a/Engine/Plugins/VirtualProduction/CinematicAssemblyTools/Source/CineAssemblyToolsEditor/Private/CineAssemblyToolsStyle.h b/Engine/Plugins/VirtualProduction/CinematicAssemblyTools/Source/CineAssemblyToolsEditor/Private/CineAssemblyToolsStyle.h
--- a/Engine/Plugins/VirtualProduction/CinematicAssemblyTools/Source/CineAssemblyToolsEditor/Private/CineAssemblyToolsStyle.h
+++ b/Engine/Plugins/VirtualProduction/CinematicAssemblyTools/Source/CineAssemblyToolsEditor/Private/CineAssemblyToolsStyle.h
@@ -5,6 +5,8 @@
 #include "Styling/SlateStyle.h"
 #include "Styling/SlateStyleRegistry.h"
 
+class UTexture2D;
+
 /** Slate style set that defines all the styles for the Cinematic Assembly Tools */
 class FCineAssemblyToolsStyle : public FSlateStyleSet
 {
@@ -14,6 +16,16 @@ public:
 	/** Access the singleton instance for this style set */
 	static FCineAssemblyToolsStyle& Get();
 
+	/** Returns the name of the thumbnail brush associated with the input schema */
+	FName GetThumbnailBrushNameForSchema(const FString& SchemaName) const;
+
+	/** 
+	 * Sets the texture resource used by the thumbnail brush associated with the input schema.
+	 * If the input schema does not have a thumbnail brush associated with it yet, a new one will be created first.
+	 * If the texture resource is null, the brush will fallback to using a predefined icon.
+	 */
+	void SetThumbnailBrushTextureForSchema(const FString& SchemaName, UTexture2D* BrushTexture);
+
 private:
 
 	FCineAssemblyToolsStyle();
diff --git a/Engine/Plugins/VirtualProduction/CinematicAssemblyTools/Source/CineAssemblyToolsEditor/Private/UI/CineAssembly/SCineAssemblyConfigPanel.cpp b/Engine/Plugins/VirtualProduction/CinematicAssemblyTools/Source/CineAssemblyToolsEditor/Private/UI/CineAssembly/SCineAssemblyConfigPanel.cpp
--- a/Engine/Plugins/VirtualProduction/CinematicAssemblyTools/Source/CineAssemblyToolsEditor/Private/UI/CineAssembly/SCineAssemblyConfigPanel.cpp
+++ b/Engine/Plugins/VirtualProduction/CinematicAssemblyTools/Source/CineAssemblyToolsEditor/Private/UI/CineAssembly/SCineAssemblyConfigPanel.cpp
@@ -2,6 +2,7 @@
 
 #include "SCineAssemblyConfigPanel.h"
 
+#include "AssetToolsModule.h"
 #include "CineAssemblyNamingTokens.h"
 #include "CineAssemblyToolsStyle.h"
 #include "IStructureDetailsView.h"
@@ -161,11 +162,7 @@ TSharedRef<SWidget> SCineAssemblyConfigPanel::MakeDetailsWidget()
 								.Padding(0.0f, 0.0f, 8.0f, 0.0f)
 								[
 									SNew(SImage)
-										.Image_Lambda([this]()
-											{
-												const UCineAssemblySchema* Schema = CineAssemblyToConfigure->GetSchema();
-												return Schema ? Schema->GetThumbnailBrush() : FCineAssemblyToolsStyle::Get().GetBrush("Thumbnails.Schema");
-											})
+										.Image(this, &SCineAssemblyConfigPanel::GetSchemaThumbnail)
 								]
 
 							+ SHorizontalBox::Slot()
@@ -225,6 +222,26 @@ bool SCineAssemblyConfigPanel::IsCustomRowVisible(FName RowName, FName ParentNam
 	return true;
 }
 
+const FSlateBrush* SCineAssemblyConfigPanel::GetSchemaThumbnail() const
+{
+	if (const UCineAssemblySchema* Schema = CineAssemblyToConfigure->GetSchema())
+	{
+		// Find the thumbnail brush associated with the selected schema
+		FAssetToolsModule& AssetToolsModule = FModuleManager::LoadModuleChecked<FAssetToolsModule>(TEXT("AssetTools"));
+		TSharedPtr<IAssetTypeActions> AssetTypeActions = AssetToolsModule.Get().GetAssetTypeActionsForClass(UCineAssemblySchema::StaticClass()).Pin();
+
+		if (AssetTypeActions)
+		{
+			const FAssetData AssetData = FAssetData(Schema);
+			const FName AssetClassName = AssetData.AssetClassPath.GetAssetName();
+
+			return AssetTypeActions->GetThumbnailBrush(AssetData, AssetClassName);
+		}
+	}
+
+	return FCineAssemblyToolsStyle::Get().GetBrush("ClassThumbnail.CineAssemblySchema");
+}
+
 TSharedRef<SWidget> SCineAssemblyConfigPanel::MakeHierarchyWidget()
 {
 	HierarchyTreeView = SNew(STreeView<TSharedPtr<FHierarchyTreeItem>>)
diff --git a/Engine/Plugins/VirtualProduction/CinematicAssemblyTools/Source/CineAssemblyToolsEditor/Private/UI/CineAssembly/SCineAssemblyConfigPanel.h b/Engine/Plugins/VirtualProduction/CinematicAssemblyTools/Source/CineAssemblyToolsEditor/Private/UI/CineAssembly/SCineAssemblyConfigPanel.h
--- a/Engine/Plugins/VirtualProduction/CinematicAssemblyTools/Source/CineAssemblyToolsEditor/Private/UI/CineAssembly/SCineAssemblyConfigPanel.h
+++ b/Engine/Plugins/VirtualProduction/CinematicAssemblyTools/Source/CineAssemblyToolsEditor/Private/UI/CineAssembly/SCineAssemblyConfigPanel.h
@@ -91,6 +91,9 @@ private:
 	/** Filter used by the Details View to determine which custom rows to show */
 	bool IsCustomRowVisible(FName RowName, FName ParentName);
 
+	/** Get the thumbnail brush for the selected schema */
+	const FSlateBrush* GetSchemaThumbnail() const;
+
 private:
 	/** Transient object used only by this UI to configure the properties of the new asset that will get created by the Factory */
 	UCineAssembly* CineAssemblyToConfigure = nullptr;
diff --git a/Engine/Plugins/VirtualProduction/CinematicAssemblyTools/Source/CineAssemblyToolsEditor/Private/UI/CineAssembly/SCineAssemblySchemaWindow.cpp b/Engine/Plugins/VirtualProduction/CinematicAssemblyTools/Source/CineAssemblyToolsEditor/Private/UI/CineAssembly/SCineAssemblySchemaWindow.cpp
--- a/Engine/Plugins/VirtualProduction/CinematicAssemblyTools/Source/CineAssemblyToolsEditor/Private/UI/CineAssembly/SCineAssemblySchemaWindow.cpp
+++ b/Engine/Plugins/VirtualProduction/CinematicAssemblyTools/Source/CineAssemblyToolsEditor/Private/UI/CineAssembly/SCineAssemblySchemaWindow.cpp
@@ -1379,6 +1379,11 @@ void SCineAssemblySchemaWindow::OnSchemaPropertiesChanged(const FPropertyChanged
 			CineAssemblyNamingTokens->AddMetadataToken(MetadataDesc.Key);
 		}
 	}
+	else if (PropertyName == GET_MEMBER_NAME_CHECKED(UCineAssemblySchema, ThumbnailImage))
+	{
+		// Update the schema's thumbnail brush in the style set to use the new thumbnail image
+		FCineAssemblyToolsStyle::Get().SetThumbnailBrushTextureForSchema(SchemaToConfigure->SchemaName, SchemaToConfigure->ThumbnailImage);
+	}
 }
 
 FReply SCineAssemblySchemaWindow::OnCreateAssetClicked()
