ShotManagement: Fix some default asset names from appearing incorrectly in the Production Wizard.
#rnx

diff --git a/Engine/Plugins/VirtualProduction/CinematicAssemblyTools/Source/CineAssemblyToolsEditor/Private/UI/SAssetNamingPanel.cpp b/Engine/Plugins/VirtualProduction/CinematicAssemblyTools/Source/CineAssemblyToolsEditor/Private/UI/SAssetNamingPanel.cpp
--- a/Engine/Plugins/VirtualProduction/CinematicAssemblyTools/Source/CineAssemblyToolsEditor/Private/UI/SAssetNamingPanel.cpp
+++ b/Engine/Plugins/VirtualProduction/CinematicAssemblyTools/Source/CineAssemblyToolsEditor/Private/UI/SAssetNamingPanel.cpp
@@ -51,6 +51,30 @@ namespace UE::CineAssemblyTools::AssetNaming::Private
 		TArray<UClass*> AllowedAssetClasses;
 	};
 
+	FString GetFactoryDefaultNameForClass(const UClass* Class)
+	{
+		const IAssetTools& AssetTools = FModuleManager::LoadModuleChecked<FAssetToolsModule>("AssetTools").Get();
+
+		FString FactoryDefaultName;
+		for (UFactory* Factory : AssetTools.GetNewAssetFactories())
+		{
+			UClass* SupportedClass = Factory->GetSupportedClass();
+
+			if (SupportedClass != nullptr && SupportedClass == Class)
+			{
+				FactoryDefaultName = Factory->GetDefaultNewAssetName();
+				break;
+			}
+		}
+
+		if (FactoryDefaultName.IsEmpty())
+		{
+			FactoryDefaultName = FString::Printf(TEXT("New%s"), *Class->GetName());
+		}
+
+		return FactoryDefaultName;
+	}
+
 	static FName ClassColumnName = TEXT("Class");
 	static FName NamingColumnName = TEXT("Naming");
 	static FName DeleteColumnName = TEXT("Delete");
@@ -337,11 +361,7 @@ void SAssetNamingRow::SetAssetClass(const UClass* SelectedClass)
 	if (SelectedClass)
 	{	
 		// Set the default name for this row to the factory default name for the selected asset type
-		const TArray<UFactory*>& Factories = PropertyCustomizationHelpers::GetNewAssetFactoriesForClasses({ SelectedClass });
-		if (Factories.Num() > 0)
-		{
-			AssetNaming->DefaultName = Factories[0]->GetDefaultNewAssetName();
-		}
+		AssetNaming->DefaultName = UE::CineAssemblyTools::AssetNaming::Private::GetFactoryDefaultNameForClass(SelectedClass);
 
 		// Let the user immediately start editing the asset name
 		FSlateApplication::Get().SetKeyboardFocus(EditableTextBlock.ToSharedRef());
