ShotManagement: Add additional input validation to assembly names.
#rnx

diff --git a/Engine/Plugins/VirtualProduction/CinematicAssemblyTools/Source/CineAssemblyToolsEditor/Private/UI/CineAssembly/SCineAssemblyConfigPanel.cpp b/Engine/Plugins/VirtualProduction/CinematicAssemblyTools/Source/CineAssemblyToolsEditor/Private/UI/CineAssembly/SCineAssemblyConfigPanel.cpp
--- a/Engine/Plugins/VirtualProduction/CinematicAssemblyTools/Source/CineAssemblyToolsEditor/Private/UI/CineAssembly/SCineAssemblyConfigPanel.cpp
+++ b/Engine/Plugins/VirtualProduction/CinematicAssemblyTools/Source/CineAssemblyToolsEditor/Private/UI/CineAssembly/SCineAssemblyConfigPanel.cpp
@@ -103,12 +103,39 @@ void SCineAssemblyConfigPanel::Construct(const FArguments& InArgs, UCineAssembly
 												CineAssemblyToConfigure->Modify();
 												CineAssemblyToConfigure->AssemblyName.Template = InText.ToString();
 											})
+										.OnValidateTokenizedText(this, &SCineAssemblyConfigPanel::ValidateAssemblyName)
 								]
 						]
 				]
 		];
 }
 
+bool SCineAssemblyConfigPanel::ValidateAssemblyName(const FText& InText, FText& OutErrorMessage) const
+{
+	// Ensure that the name does not contain any characters that would be invalid for an asset name
+	// This matches the validation that would happen if the user was renaming an asset in the content browser
+	FString InvalidCharacters = INVALID_OBJECTNAME_CHARACTERS INVALID_LONGPACKAGE_CHARACTERS;
+
+	// These characters are actually valid, because we want to support naming tokens
+	InvalidCharacters = InvalidCharacters.Replace(TEXT("{}"), TEXT(""));
+	InvalidCharacters = InvalidCharacters.Replace(TEXT(":"), TEXT(""));
+
+	const FString PotentialName = InText.ToString();
+
+	if (!FName::IsValidXName(PotentialName, InvalidCharacters, &OutErrorMessage))
+	{
+		return false;
+	}
+
+	if (PotentialName.Contains(TEXT("{assembly}")))
+	{
+		OutErrorMessage = LOCTEXT("RecursiveAssemblyTokenError", "You cannot use the {assembly} token");
+		return false;
+	}
+
+	return true;
+}
+
 void SCineAssemblyConfigPanel::Refresh()
 {
 	// The details view needs to be redrawn to show the new metadata fields from the selected schema
diff --git a/Engine/Plugins/VirtualProduction/CinematicAssemblyTools/Source/CineAssemblyToolsEditor/Private/UI/CineAssembly/SCineAssemblyConfigPanel.h b/Engine/Plugins/VirtualProduction/CinematicAssemblyTools/Source/CineAssemblyToolsEditor/Private/UI/CineAssembly/SCineAssemblyConfigPanel.h
--- a/Engine/Plugins/VirtualProduction/CinematicAssemblyTools/Source/CineAssemblyToolsEditor/Private/UI/CineAssembly/SCineAssemblyConfigPanel.h
+++ b/Engine/Plugins/VirtualProduction/CinematicAssemblyTools/Source/CineAssemblyToolsEditor/Private/UI/CineAssembly/SCineAssemblyConfigPanel.h
@@ -82,6 +82,9 @@ private:
 	/** Recursively evaluates the tokens in the path of each item in the tree view */
 	void EvaluateHierarchyTokensRecursive(TSharedPtr<FHierarchyTreeItem> TreeItem);
 
+	/** Validates the user input text for the assembly name */
+	bool ValidateAssemblyName(const FText& InText, FText& OutErrorMessage) const;
+
 	/**
 	 * Evaluates the input template string with the naming tokens subsystem, and stores the result in the Resolved text.
 	 * This function is throttled to only run at a set frequency, to avoid the potential to constantly query the naming tokens subsystem.
diff --git a/Engine/Source/Editor/EditorWidgets/Private/STemplateStringEditableTextBox.cpp b/Engine/Source/Editor/EditorWidgets/Private/STemplateStringEditableTextBox.cpp
--- a/Engine/Source/Editor/EditorWidgets/Private/STemplateStringEditableTextBox.cpp
+++ b/Engine/Source/Editor/EditorWidgets/Private/STemplateStringEditableTextBox.cpp
@@ -370,7 +370,10 @@ void STemplateStringEditableTextBox::OnEditableTextChanged(const FText& InTokeni
 
 void STemplateStringEditableTextBox::OnEditableTextCommitted(const FText& InTokenizedText, ETextCommit::Type InCommitType)
 {
-	ValidateTokenizedText(InTokenizedText);
+	if (!ValidateTokenizedText(InTokenizedText))
+	{
+		return;
+	}
 
 	if (!TokenizedText.IsBound())
 	{
