ShotManagement: Fix schemas to correctly set the SchemaName property when duplicated.
#rnx

diff --git a/Engine/Plugins/VirtualProduction/CinematicAssemblyTools/Source/CineAssemblyTools/Private/CineAssemblySchema.cpp b/Engine/Plugins/VirtualProduction/CinematicAssemblyTools/Source/CineAssemblyTools/Private/CineAssemblySchema.cpp
--- a/Engine/Plugins/VirtualProduction/CinematicAssemblyTools/Source/CineAssemblyTools/Private/CineAssemblySchema.cpp
+++ b/Engine/Plugins/VirtualProduction/CinematicAssemblyTools/Source/CineAssemblyTools/Private/CineAssemblySchema.cpp
@@ -47,6 +47,9 @@ void UCineAssemblySchema::PostDuplicate(bool bDuplicateForPIE)
 	if (!bDuplicateForPIE)
 	{
 		SchemaGuid = FGuid::NewGuid();
+
+		// Change the name property to match the actual name of the duplicated asset
+		SchemaName = this->GetName();
 	}
 }
 
diff --git a/Engine/Plugins/VirtualProduction/CinematicAssemblyTools/Source/CineAssemblyToolsEditor/Private/CineAssemblySchemaCustomization.cpp b/Engine/Plugins/VirtualProduction/CinematicAssemblyTools/Source/CineAssemblyToolsEditor/Private/CineAssemblySchemaCustomization.cpp
--- a/Engine/Plugins/VirtualProduction/CinematicAssemblyTools/Source/CineAssemblyToolsEditor/Private/CineAssemblySchemaCustomization.cpp
+++ b/Engine/Plugins/VirtualProduction/CinematicAssemblyTools/Source/CineAssemblyToolsEditor/Private/CineAssemblySchemaCustomization.cpp
@@ -65,7 +65,13 @@ void FCineAssemblySchemaCustomization::CustomizeDetails(IDetailLayoutBuilder& De
 					})
 				.OnTextCommitted_Lambda([this](const FText& InText, ETextCommit::Type InCommitType)
 					{
-						CustomizedSchema->RenameAsset(InText.ToString());
+						if (InCommitType != ETextCommit::Default)
+						{
+							if (!CustomizedSchema->SchemaName.Equals(InText.ToString()))
+							{
+								CustomizedSchema->RenameAsset(InText.ToString());
+							}
+						}
 					})
 				.OnVerifyTextChanged_Raw(this, &FCineAssemblySchemaCustomization::ValidateSchemaName)
 		];
