ShotManagement: Restrict production names to have a maximum length of 256 characters.
#rnx

diff --git a/Engine/Plugins/VirtualProduction/CinematicAssemblyTools/Source/CineAssemblyToolsEditor/Private/CineAssemblyToolsEditorModule.cpp b/Engine/Plugins/VirtualProduction/CinematicAssemblyTools/Source/CineAssemblyToolsEditor/Private/CineAssemblyToolsEditorModule.cpp
--- a/Engine/Plugins/VirtualProduction/CinematicAssemblyTools/Source/CineAssemblyToolsEditor/Private/CineAssemblyToolsEditorModule.cpp
+++ b/Engine/Plugins/VirtualProduction/CinematicAssemblyTools/Source/CineAssemblyToolsEditor/Private/CineAssemblyToolsEditorModule.cpp
@@ -12,6 +12,7 @@
 #include "CineAssemblySchema.h"
 #include "CineAssemblySchemaCustomization.h"
 #include "CineAssemblyToolsStyle.h"
+#include "CinematicProductionCustomization.h"
 #include "ContentBrowserMenuContexts.h"
 #include "Editor.h"
 #include "Framework/Application/SlateApplication.h"
@@ -62,6 +63,11 @@ void FCineAssemblyToolsEditorModule::StartupModule()
 		FOnGetDetailCustomizationInstance::CreateStatic(&FCineAssemblySchemaCustomization::MakeInstance)
 	);
 
+	PropertyEditorModule.RegisterCustomPropertyTypeLayout(
+		FCinematicProduction::StaticStruct()->GetFName(),
+		FOnGetPropertyTypeCustomizationInstance::CreateStatic(&FCinematicProductionCustomization::MakeInstance)
+	);
+
 	PropertyEditorModule.RegisterCustomPropertyTypeLayout(
 		FAssemblyMetadataDesc::StaticStruct()->GetFName(),
 		FOnGetPropertyTypeCustomizationInstance::CreateStatic(&FCineAssemblyMetadataCustomization::MakeInstance)
diff --git a/Engine/Plugins/VirtualProduction/CinematicAssemblyTools/Source/CineAssemblyToolsEditor/Private/CinematicProductionCustomization.cpp b/Engine/Plugins/VirtualProduction/CinematicAssemblyTools/Source/CineAssemblyToolsEditor/Private/CinematicProductionCustomization.cpp
new file mode 100644
--- /dev/null
+++ b/Engine/Plugins/VirtualProduction/CinematicAssemblyTools/Source/CineAssemblyToolsEditor/Private/CinematicProductionCustomization.cpp
@@ -0,0 +1,72 @@
+// Copyright Epic Games, Inc. All Rights Reserved.
+
+#include "CinematicProductionCustomization.h"
+
+//#include "DetailLayoutBuilder.h"
+#include "DetailWidgetRow.h"
+//#include "Framework/MultiBox/MultiBoxBuilder.h"
+#include "IDetailChildrenBuilder.h"
+//#include "Modules/ModuleManager.h"
+#include "ProductionSettings.h"
+//#include "Widgets/Input/SComboButton.h"
+#include "Widgets/Input/SEditableTextBox.h"
+
+TSharedRef<IPropertyTypeCustomization> FCinematicProductionCustomization::MakeInstance()
+{
+	return MakeShareable(new FCinematicProductionCustomization);
+}
+
+void FCinematicProductionCustomization::CustomizeHeader(TSharedRef<IPropertyHandle> InPropertyHandle, FDetailWidgetRow& HeaderRow, IPropertyTypeCustomizationUtils& CustomizationUtils)
+{
+	HeaderRow.NameContent()
+		[
+			InPropertyHandle->CreatePropertyNameWidget()
+		]
+		.ValueContent()
+		[
+			InPropertyHandle->CreatePropertyValueWidget()
+		];
+}
+
+void FCinematicProductionCustomization::CustomizeChildren(TSharedRef<IPropertyHandle> PropertyHandle, IDetailChildrenBuilder& ChildBuilder, IPropertyTypeCustomizationUtils& CustomizationUtils)
+{
+	uint32 NumChildren;
+	PropertyHandle->GetNumChildren(NumChildren);
+
+	for (uint32 Index = 0; Index < NumChildren; ++Index)
+	{
+		TSharedRef<IPropertyHandle> ChildPropertyHandle = PropertyHandle->GetChildHandle(Index).ToSharedRef();
+
+		// Replace the default textbox widget that would be used for this FString property with one that will limit the maximum length of the Production Name
+		if (ChildPropertyHandle->GetProperty()->GetFName() == GET_MEMBER_NAME_CHECKED(FCinematicProduction, ProductionName))
+		{
+			IDetailPropertyRow& ProductionNameRow = ChildBuilder.AddProperty(ChildPropertyHandle);
+			ProductionNameRow.CustomWidget()
+				.NameContent()
+				[
+					ChildPropertyHandle->CreatePropertyNameWidget()
+				]
+				.ValueContent()
+				[
+					SNew(SEditableTextBox)
+						.Font(CustomizationUtils.GetRegularFont())
+						.MaximumLength(UProductionSettings::ProductionNameMaxLength)
+						.Text_Lambda([this, ChildPropertyHandle]() -> FText
+							{
+								FString ProductionName;
+								ChildPropertyHandle->GetValue(ProductionName);
+								return FText::FromString(ProductionName);
+							})
+						.OnTextCommitted_Lambda([this, ChildPropertyHandle](const FText& InText, ETextCommit::Type InCommitType)
+							{
+								const FString NewProductionName = InText.ToString();
+								ChildPropertyHandle->SetValue(NewProductionName);
+							})
+				];
+		}
+		else
+		{
+			ChildBuilder.AddProperty(ChildPropertyHandle);
+		}
+	}
+}
diff --git a/Engine/Plugins/VirtualProduction/CinematicAssemblyTools/Source/CineAssemblyToolsEditor/Private/CinematicProductionCustomization.h b/Engine/Plugins/VirtualProduction/CinematicAssemblyTools/Source/CineAssemblyToolsEditor/Private/CinematicProductionCustomization.h
new file mode 100644
--- /dev/null
+++ b/Engine/Plugins/VirtualProduction/CinematicAssemblyTools/Source/CineAssemblyToolsEditor/Private/CinematicProductionCustomization.h
@@ -0,0 +1,22 @@
+// Copyright Epic Games, Inc. All Rights Reserved.
+
+#pragma once
+
+#include "IPropertyTypeCustomization.h"
+//#include "IPropertyUtilities.h"
+
+//class IDetailLayoutBuilder;
+
+/**
+ * Customization for FCinematicProduction
+ */
+class FCinematicProductionCustomization : public IPropertyTypeCustomization
+{
+public:
+	static TSharedRef<IPropertyTypeCustomization> MakeInstance();
+
+	// Begin IDetailCustomization interface
+	virtual void CustomizeHeader(TSharedRef<IPropertyHandle> PropertyHandle, FDetailWidgetRow& HeaderRow, IPropertyTypeCustomizationUtils& CustomizationUtils) override;
+	virtual void CustomizeChildren(TSharedRef<IPropertyHandle> PropertyHandle, IDetailChildrenBuilder& ChildBuilder, IPropertyTypeCustomizationUtils& CustomizationUtils) override;
+	//~ End IDetailCustomization interface
+};
diff --git a/Engine/Plugins/VirtualProduction/CinematicAssemblyTools/Source/CineAssemblyToolsEditor/Private/ProductionSettings.cpp b/Engine/Plugins/VirtualProduction/CinematicAssemblyTools/Source/CineAssemblyToolsEditor/Private/ProductionSettings.cpp
--- a/Engine/Plugins/VirtualProduction/CinematicAssemblyTools/Source/CineAssemblyToolsEditor/Private/ProductionSettings.cpp
+++ b/Engine/Plugins/VirtualProduction/CinematicAssemblyTools/Source/CineAssemblyToolsEditor/Private/ProductionSettings.cpp
@@ -525,6 +525,11 @@ void UProductionSettings::PostEditUndo()
 	ReloadProductionsExtendedData();
 }
 
+bool UProductionSettings::IsValidProductionName(const FString& ProductionName) const
+{
+	return ProductionName.Len() <= ProductionNameMaxLength;
+}
+
 void UProductionSettings::AddProduction()
 {
 	FCinematicProduction NewProduction;
@@ -534,6 +539,12 @@ void UProductionSettings::AddProduction()
 
 void UProductionSettings::AddProduction(const FCinematicProduction& ProductionToAdd)
 {
+	if (!IsValidProductionName(ProductionToAdd.ProductionName))
+	{
+		UE_LOG(LogProductionSettings, Error, TEXT("Failed to add Production named %s because the name is too long. The maximum production name length is %d"), *ProductionToAdd.ProductionName, ProductionNameMaxLength);
+		return;
+	}
+
 	Productions.Add(ProductionToAdd);
 
 	if (Productions.Num() == 1)
@@ -558,6 +569,12 @@ void UProductionSettings::DuplicateProduction(FGuid ProductionID)
 		DuplicateProduction.ProductionID = FGuid::NewGuid();
 		DuplicateProduction.ProductionName = GetUniqueProductionName(ProductionToCopy->ProductionName);
 
+		if (!IsValidProductionName(DuplicateProduction.ProductionName))
+		{
+			UE_LOG(LogProductionSettings, Error, TEXT("Failed to duplicate Production named %s because the name is too long. The maximum production name length is %d"), *ProductionToCopy->ProductionName, ProductionNameMaxLength);
+			return;
+		}
+
 		Productions.Add(DuplicateProduction);
 
 		UpdateConfig();
@@ -686,6 +703,12 @@ void UProductionSettings::DeleteProduction(FGuid ProductionID)
 
 void UProductionSettings::RenameProduction(FGuid ProductionID, FString NewName)
 {
+	if (!IsValidProductionName(NewName))
+	{
+		UE_LOG(LogProductionSettings, Error, TEXT("Failed to rename Production to %s because the name is too long. The maximum production name length is %d"), *NewName, ProductionNameMaxLength);
+		return;
+	}
+
 	if (FCinematicProduction* Production = Algo::FindBy(Productions, ProductionID, &FCinematicProduction::ProductionID))
 	{
 		Production->ProductionName = NewName;
diff --git a/Engine/Plugins/VirtualProduction/CinematicAssemblyTools/Source/CineAssemblyToolsEditor/Private/UI/SProductionListPanel.cpp b/Engine/Plugins/VirtualProduction/CinematicAssemblyTools/Source/CineAssemblyToolsEditor/Private/UI/SProductionListPanel.cpp
--- a/Engine/Plugins/VirtualProduction/CinematicAssemblyTools/Source/CineAssemblyToolsEditor/Private/UI/SProductionListPanel.cpp
+++ b/Engine/Plugins/VirtualProduction/CinematicAssemblyTools/Source/CineAssemblyToolsEditor/Private/UI/SProductionListPanel.cpp
@@ -487,6 +487,7 @@ void SProductionListButton::Construct(const FArguments& InArgs)
 						.OnVerifyTextChanged_Lambda(OnVerifyRename)
 						.OnTextCommitted_Lambda(OnRenameFinished)
 						.IsSelected_Lambda([]() { return false; }) // Disable double-select to rename
+						.MaximumLength(UProductionSettings::ProductionNameMaxLength)
 				]
 		];
 }
diff --git a/Engine/Plugins/VirtualProduction/CinematicAssemblyTools/Source/CineAssemblyToolsEditor/Public/ProductionSettings.h b/Engine/Plugins/VirtualProduction/CinematicAssemblyTools/Source/CineAssemblyToolsEditor/Public/ProductionSettings.h
--- a/Engine/Plugins/VirtualProduction/CinematicAssemblyTools/Source/CineAssemblyToolsEditor/Public/ProductionSettings.h
+++ b/Engine/Plugins/VirtualProduction/CinematicAssemblyTools/Source/CineAssemblyToolsEditor/Public/ProductionSettings.h
@@ -203,6 +203,9 @@ class UProductionSettings : public UDeveloperSettings
 public:
 	UProductionSettings() = default;
 
+	/** Maximum number of characters allowed in a Production Name */
+	UE_API constexpr static int32 ProductionNameMaxLength = 256;
+
 	//~ Begin UDeveloperSettings overrides
 	virtual FName GetContainerName() const override;
 	virtual FName GetCategoryName() const override;
@@ -408,6 +411,9 @@ public:
 	UE_API static FName GetProductionsPropertyMemberName();
 
 private:
+	/** Returns true if the input string is a valid Production name */
+	bool IsValidProductionName(const FString& ProductionName) const;
+
 	/** Applies overrides to various project settings based on the active production settings */
 	void ApplyProjectOverrides();
 
