ShotManagement: Fix subassemblies to transact properly over multi-user when they are newly created by delaying when the asset registry is notified of its creation until after it has been fully initialized.
#rnx

diff --git a/Engine/Plugins/VirtualProduction/CinematicAssemblyTools/Source/CineAssemblyTools/Private/CineAssembly.cpp b/Engine/Plugins/VirtualProduction/CinematicAssemblyTools/Source/CineAssemblyTools/Private/CineAssembly.cpp
--- a/Engine/Plugins/VirtualProduction/CinematicAssemblyTools/Source/CineAssemblyTools/Private/CineAssembly.cpp
+++ b/Engine/Plugins/VirtualProduction/CinematicAssemblyTools/Source/CineAssemblyTools/Private/CineAssembly.cpp
@@ -339,7 +339,12 @@ void UCineAssembly::CreateSubAssemblies()
 		AssetTools.CreateUniqueAssetName(AssemblyRootFolder / SubAssemblyName.Resolved.ToString(), TEXT(""), UniquePackageName, UniqueAssetName);
 
 		const FString SubAssemblyPath = FPaths::GetPath(UniquePackageName);
-		UCineAssembly* SubAssembly = Cast<UCineAssembly>(AssetTools.CreateAsset(UniqueAssetName, SubAssemblyPath, UCineAssembly::StaticClass(), nullptr));
+
+		UPackage* SubAssemblyPackage = CreatePackage(*UniquePackageName);
+
+		const EObjectFlags Flags = RF_Public | RF_Standalone | RF_Transactional;
+		UCineAssembly* SubAssembly = NewObject<UCineAssembly>(SubAssemblyPackage, UCineAssembly::StaticClass(), FName(*UniqueAssetName), Flags);
+
 		if (SubAssembly)
 		{
 			ULevelSequence* Template = nullptr;
@@ -373,6 +378,12 @@ void UCineAssembly::CreateSubAssemblies()
 
 			UMovieSceneSubSection* SubSection = SubTrack->AddSequence(SubAssembly, StartFrame, Duration);
 			SubAssemblies.Add(SubSection);
+
+			// Notify the asset registry about the new subassembly
+			FAssetRegistryModule::AssetCreated(SubAssembly);
+
+			// Mark the package dirty
+			SubAssemblyPackage->MarkPackageDirty();
 		}
 	}
 }
