ShotManagement: Deleting commented out code
#rnx

diff --git a/Engine/Plugins/VirtualProduction/CinematicAssemblyTools/Source/CineAssemblyToolsEditor/Private/CinematicProductionCustomization.cpp b/Engine/Plugins/VirtualProduction/CinematicAssemblyTools/Source/CineAssemblyToolsEditor/Private/CinematicProductionCustomization.cpp
--- a/Engine/Plugins/VirtualProduction/CinematicAssemblyTools/Source/CineAssemblyToolsEditor/Private/CinematicProductionCustomization.cpp
+++ b/Engine/Plugins/VirtualProduction/CinematicAssemblyTools/Source/CineAssemblyToolsEditor/Private/CinematicProductionCustomization.cpp
@@ -2,13 +2,9 @@
 
 #include "CinematicProductionCustomization.h"
 
-//#include "DetailLayoutBuilder.h"
 #include "DetailWidgetRow.h"
-//#include "Framework/MultiBox/MultiBoxBuilder.h"
 #include "IDetailChildrenBuilder.h"
-//#include "Modules/ModuleManager.h"
 #include "ProductionSettings.h"
-//#include "Widgets/Input/SComboButton.h"
 #include "Widgets/Input/SEditableTextBox.h"
 
 TSharedRef<IPropertyTypeCustomization> FCinematicProductionCustomization::MakeInstance()
diff --git a/Engine/Plugins/VirtualProduction/CinematicAssemblyTools/Source/CineAssemblyToolsEditor/Private/CinematicProductionCustomization.h b/Engine/Plugins/VirtualProduction/CinematicAssemblyTools/Source/CineAssemblyToolsEditor/Private/CinematicProductionCustomization.h
--- a/Engine/Plugins/VirtualProduction/CinematicAssemblyTools/Source/CineAssemblyToolsEditor/Private/CinematicProductionCustomization.h
+++ b/Engine/Plugins/VirtualProduction/CinematicAssemblyTools/Source/CineAssemblyToolsEditor/Private/CinematicProductionCustomization.h
@@ -3,9 +3,6 @@
 #pragma once
 
 #include "IPropertyTypeCustomization.h"
-//#include "IPropertyUtilities.h"
-
-//class IDetailLayoutBuilder;
 
 /**
  * Customization for FCinematicProduction
