ShotManagement: Add support to the Production Wizard to specify template folder paths in plugins
#rnx

diff --git a/Engine/Plugins/VirtualProduction/CinematicAssemblyTools/Source/CineAssemblyToolsEditor/Private/UI/SFolderHierarchyPanel.cpp b/Engine/Plugins/VirtualProduction/CinematicAssemblyTools/Source/CineAssemblyToolsEditor/Private/UI/SFolderHierarchyPanel.cpp
--- a/Engine/Plugins/VirtualProduction/CinematicAssemblyTools/Source/CineAssemblyToolsEditor/Private/UI/SFolderHierarchyPanel.cpp
+++ b/Engine/Plugins/VirtualProduction/CinematicAssemblyTools/Source/CineAssemblyToolsEditor/Private/UI/SFolderHierarchyPanel.cpp
@@ -14,8 +14,9 @@
 #include "Engine/Engine.h"
 #include "Framework/Application/SlateApplication.h"
 #include "Framework/MultiBox/MultiBoxBuilder.h"
-#include "IContentBrowserDataModule.h"
 #include "HAL/FileManager.h"
+#include "IContentBrowserDataModule.h"
+#include "Interfaces/IPluginManager.h"
 #include "Misc/MessageDialog.h"
 #include "NamingTokensEngineSubsystem.h"
 #include "ObjectTools.h"
@@ -31,8 +32,6 @@
 
 #define LOCTEXT_NAMESPACE "SFolderHierarchyPanel"
 
-FName SFolderHierarchyPanel::RootPath = TEXT("/Game");
-
 void SFolderHierarchyPanel::Construct(const FArguments& InArgs)
 {
 	// Subscribe to be notified when the Production Settings active productions has changed
@@ -40,9 +39,17 @@ void SFolderHierarchyPanel::Construct(const FArguments& InArgs)
 	ActiveProductionChangedHandle = ProductionSettings->OnActiveProductionChanged().AddSP(this, &SFolderHierarchyPanel::BuildTreeFromProductionTemplate);
 
 	// Create the tree root
-	RootItem = MakeShared<FTemplateFolderTreeItem>();
-	RootItem->Path.SetPathFromString(RootPath.ToString(), EContentBrowserPathType::Internal);
-	RootItem->Status = ETemplateFolderStatus::Exists;
+	ContentRootItem = MakeShared<FTemplateFolderTreeItem>();
+	ContentRootItem->Path.SetPathFromString(TEXT("/Game"), EContentBrowserPathType::Internal);
+	ContentRootItem->FriendlyName = TEXT("Content");
+	ContentRootItem->Status = ETemplateFolderStatus::Exists;
+	ContentRootItem->bIsReadOnly = true;
+
+	PluginRootItem = MakeShared<FTemplateFolderTreeItem>();
+	PluginRootItem->Path.SetPathFromString(TEXT("/"), EContentBrowserPathType::Internal);
+	PluginRootItem->FriendlyName = TEXT("Plugins");
+	PluginRootItem->Status = ETemplateFolderStatus::Exists;
+	PluginRootItem->bIsReadOnly = true;
 
 	// Initialize the tree view using the active production's folder template
 	BuildTreeFromProductionTemplate();
@@ -54,11 +61,13 @@ void SFolderHierarchyPanel::Construct(const FArguments& InArgs)
 		.OnGetChildren(this, &SFolderHierarchyPanel::OnGetChildren)
 		.OnItemsRebuilt(this, &SFolderHierarchyPanel::OnItemsRebuilt)
 		.OnKeyDownHandler(this, &SFolderHierarchyPanel::OnKeyDown)
+		.OnMouseButtonDoubleClick(this, &SFolderHierarchyPanel::OnTreeViewDoubleClick)
 		.OnSetExpansionRecursive(this, &SFolderHierarchyPanel::SetExpansionRecursive)
 		.OnContextMenuOpening(this, &SFolderHierarchyPanel::OnContextMenuOpening);
 
 	// Start with the entire template tree expanded
-	SetExpansionRecursive(RootItem, true);
+	SetExpansionRecursive(ContentRootItem, true);
+	SetExpansionRecursive(PluginRootItem, true);
 
 	ChildSlot
 		[
@@ -282,17 +291,10 @@ TSharedRef<ITableRow> SFolderHierarchyPanel::OnGenerateRow(TSharedPtr<FTemplateF
 					SAssignNew(TreeItem->NameWidget, SInlineEditableTextBlock)
 						.Text_Lambda([TreeItem]() -> FText
 							{
-								// The Content Browser overrides the display name of the top-level project content folder to "Content", which we want to match
-								if (TreeItem->Path.GetInternalPathName() == RootPath)
-								{
-									return LOCTEXT("ContentFolderDisplayNameOverride", "Content");
-								}
-								return FText::FromString(FPaths::GetPathLeaf(TreeItem->Path.GetInternalPathString()));
-							})
-						.IsReadOnly_Lambda([TreeItem]() -> bool
-							{
-								return (TreeItem->Path.GetInternalPathName() == RootPath);
+								const FString ItemName = !TreeItem->FriendlyName.IsEmpty() ? TreeItem->FriendlyName : FPaths::GetPathLeaf(TreeItem->Path.GetInternalPathString());
+								return FText::FromString(ItemName);
 							})
+						.IsReadOnly(TreeItem->bIsReadOnly)
 						.OnVerifyTextChanged(this, &SFolderHierarchyPanel::IsValidFolderName, TreeItem)
 						.OnTextCommitted(this, &SFolderHierarchyPanel::SetTemplateFolderName, TreeItem)
 				]
@@ -346,44 +348,44 @@ TSharedPtr<SWidget> SFolderHierarchyPanel::OnContextMenuOpening()
 
 		TSharedPtr<FTemplateFolderTreeItem>& SelectedTreeItem = SelectedNodes[0];
 
-		if (SelectedTreeItem->Path.GetInternalPathName() != RootPath)
+		// If the selected template folder does not currently exist in the Content Browser, give the user the option to toggle its status between "Create" and "Do Not Create"
+		if (SelectedTreeItem->Status == ETemplateFolderStatus::MissingCreate)
 		{
-			// If the selected template folder does not currently exist in the Content Browser, give the user the option to toggle its status between "Create" and "Do Not Create"
-			if (SelectedTreeItem->Status == ETemplateFolderStatus::MissingCreate)
-			{
-				MenuBuilder.AddMenuEntry(
-					LOCTEXT("MarkAsDoNotCreate", "Mark as Do Not Create"),
-					LOCTEXT("MarkAsDoNotCreateToolTip", "Mark this folder so that it will not be created when clicking Create Template Folders"),
-					FSlateIcon(FAppStyle::GetAppStyleSetName(), "SourceControl.StatusIcon.Off"),
-					FUIAction(FExecuteAction::CreateSP(this, &SFolderHierarchyPanel::SetTemplateFolderStatusRecursive, SelectedTreeItem, ETemplateFolderStatus::MissingDoNotCreate)),
-					NAME_None,
-					EUserInterfaceActionType::Button
-				);
-			}
-			else if (SelectedTreeItem->Status == ETemplateFolderStatus::MissingDoNotCreate)
-			{
-				MenuBuilder.AddMenuEntry(
-					LOCTEXT("MarkAsCreate", "Mark as Create"),
-					LOCTEXT("MarkAsCreateToolTip", "Mark this folder so that it will not be created when clicking Create Template Folders"),
-					FSlateIcon(FAppStyle::GetAppStyleSetName(), "Icons.AddCircle"),
-					FUIAction(FExecuteAction::CreateSP(this, &SFolderHierarchyPanel::SetTemplateFolderStatusRecursive, SelectedTreeItem, ETemplateFolderStatus::MissingCreate)),
-					NAME_None,
-					EUserInterfaceActionType::Button
-				);
-			}
+			MenuBuilder.AddMenuEntry(
+				LOCTEXT("MarkAsDoNotCreate", "Mark as Do Not Create"),
+				LOCTEXT("MarkAsDoNotCreateToolTip", "Mark this folder so that it will not be created when clicking Create Template Folders"),
+				FSlateIcon(FAppStyle::GetAppStyleSetName(), "SourceControl.StatusIcon.Off"),
+				FUIAction(FExecuteAction::CreateSP(this, &SFolderHierarchyPanel::SetTemplateFolderStatusRecursive, SelectedTreeItem, ETemplateFolderStatus::MissingDoNotCreate)),
+				NAME_None,
+				EUserInterfaceActionType::Button
+			);
+		}
+		else if (SelectedTreeItem->Status == ETemplateFolderStatus::MissingDoNotCreate)
+		{
+			MenuBuilder.AddMenuEntry(
+				LOCTEXT("MarkAsCreate", "Mark as Create"),
+				LOCTEXT("MarkAsCreateToolTip", "Mark this folder so that it will not be created when clicking Create Template Folders"),
+				FSlateIcon(FAppStyle::GetAppStyleSetName(), "Icons.AddCircle"),
+				FUIAction(FExecuteAction::CreateSP(this, &SFolderHierarchyPanel::SetTemplateFolderStatusRecursive, SelectedTreeItem, ETemplateFolderStatus::MissingCreate)),
+				NAME_None,
+				EUserInterfaceActionType::Button
+			);
 		}
 
 		// Option to add a new child folder
-		MenuBuilder.AddMenuEntry(
-			LOCTEXT("AddChildFolderAction", "Add Child Folder"),
-			LOCTEXT("AddChildFolderTooltip", "Add child folder"),
-			FSlateIcon(FAppStyle::GetAppStyleSetName(), "Icons.Plus"),
-			FUIAction(FExecuteAction::CreateLambda([this]() { OnAddFolderToTemplate(); })),
-			NAME_None,
-			EUserInterfaceActionType::Button
-		);
-
-		if (SelectedTreeItem->Path.GetInternalPathName() != RootPath)
+		if (SelectedTreeItem != PluginRootItem)
+		{
+			MenuBuilder.AddMenuEntry(
+				LOCTEXT("AddChildFolderAction", "Add Child Folder"),
+				LOCTEXT("AddChildFolderTooltip", "Add child folder"),
+				FSlateIcon(FAppStyle::GetAppStyleSetName(), "Icons.Plus"),
+				FUIAction(FExecuteAction::CreateLambda([this]() { OnAddFolderToTemplate(); })),
+				NAME_None,
+				EUserInterfaceActionType::Button
+			);
+		}
+
+		if (!SelectedTreeItem->bIsReadOnly)
 		{
 			// Option to rename the template folder (and update the path of the children below it)
 			MenuBuilder.AddMenuEntry(
@@ -425,7 +427,7 @@ FReply SFolderHierarchyPanel::OnKeyDown(const FGeometry& MyGeometry, const FKeyE
 	{
 		TSharedPtr<FTemplateFolderTreeItem>& SelectedTreeItem = SelectedNodes[0];
 
-		if (InKeyEvent.GetKey() == EKeys::Delete)
+		if ((InKeyEvent.GetKey() == EKeys::Delete) && (!SelectedTreeItem->bIsReadOnly))
 		{
 			DeleteTemplateFolder(SelectedTreeItem);
 		}
@@ -434,12 +436,50 @@ FReply SFolderHierarchyPanel::OnKeyDown(const FGeometry& MyGeometry, const FKeyE
 	return FReply::Handled();
 }
 
+void SFolderHierarchyPanel::OnTreeViewDoubleClick(TSharedPtr<FTemplateFolderTreeItem> TreeItem)
+{
+	// Put the textblock for the input tree item into edit mode so the user can rename the item
+	if (TreeItem && TreeItem->NameWidget)
+	{
+		FSlateApplication::Get().SetKeyboardFocus(TreeItem->NameWidget.ToSharedRef());
+		TreeItem->NameWidget->EnterEditingMode();
+	}
+}
+
 void SFolderHierarchyPanel::BuildTreeFromProductionTemplate()
 {
 	FolderItemsSource.Empty();
-	RootItem->Children.Empty();
+	ContentRootItem->Children.Empty();
+	PluginRootItem->Children.Empty();
+
+	FolderItemsSource.Add(ContentRootItem);
+
+	TArray<TSharedRef<IPlugin>> EnabledPlugins = IPluginManager::Get().GetEnabledPluginsWithContent();
 
-	FolderItemsSource.Add(RootItem);
+	for (const TSharedRef<IPlugin>& Plugin : EnabledPlugins)
+	{
+		if (Plugin.Get().GetType() != EPluginType::Engine)
+		{
+			FString PluginContentPath = Plugin.Get().GetMountedAssetPath();
+			PluginContentPath.RemoveFromEnd(TEXT("/"));
+
+			TSharedPtr<FTemplateFolderTreeItem> PluginItem = MakeShared<FTemplateFolderTreeItem>();
+			PluginItem->Path.SetPathFromString(PluginContentPath, EContentBrowserPathType::Internal);
+			PluginItem->FriendlyName = Plugin.Get().GetFriendlyName();
+			PluginItem->Status = ETemplateFolderStatus::Exists;
+			PluginItem->bIsReadOnly = true;
+
+			PluginRootItem->Children.Add(PluginItem);
+		}
+	}
+
+	// Don't add the plugin root folder to the tree if there are no actual plugins enabled
+	if (!PluginRootItem->Children.IsEmpty())
+	{
+		FolderItemsSource.Add(PluginRootItem);
+	}
+
+	PluginRootItem->Children.Sort([](const TSharedPtr<FTemplateFolderTreeItem>& A, const TSharedPtr<FTemplateFolderTreeItem>& B) { return A->FriendlyName < B->FriendlyName; });
 
 	const UProductionSettings* const ProductionSettings = GetMutableDefault<UProductionSettings>();
 	TOptional<const FCinematicProduction> ActiveProduction = ProductionSettings->GetActiveProduction();
@@ -457,8 +497,7 @@ void SFolderHierarchyPanel::BuildTreeFromProductionTemplate()
 
 			// Walk the tree until we find an item whose path matches the parent path. The new tree item will be created as one of its children
 			// If no parent is found, we do not add the template folder to the tree view because it is considered ill-formed
-			TSharedPtr<FTemplateFolderTreeItem> ParentItem = FindItemAtPathRecursive(RootItem, ParentPath);
-			if (ParentItem)
+			if (TSharedPtr<FTemplateFolderTreeItem> ParentItem = FindItemAtPath(ParentPath))
 			{
 				TSharedPtr<FTemplateFolderTreeItem> NewItem = MakeShared<FTemplateFolderTreeItem>();
 				NewItem->Path.SetPathFromString(FolderTemplate.InternalPath, EContentBrowserPathType::Internal);
@@ -481,7 +520,8 @@ void SFolderHierarchyPanel::BuildTreeFromProductionTemplate()
 	if (TreeView)
 	{
 		TreeView->RequestTreeRefresh();
-		SetExpansionRecursive(RootItem, true);
+		SetExpansionRecursive(ContentRootItem, true);
+		SetExpansionRecursive(PluginRootItem, true);
 	}
 }
 
@@ -489,7 +529,7 @@ FReply SFolderHierarchyPanel::OnAddFolderToTemplate()
 {
 	// Get the parent item for the new template folder being added (this can be the root folder if no parent is currently selected)
 	// The TreeView uses single selection mode, so at most one item can ever be selected by the user
-	TSharedPtr<FTemplateFolderTreeItem> SelectedTreeItem = RootItem;
+	TSharedPtr<FTemplateFolderTreeItem> SelectedTreeItem = ContentRootItem;
 
 	TArray<TSharedPtr<FTemplateFolderTreeItem>> SelectedNodes = TreeView->GetSelectedItems();
 	if (SelectedNodes.Num() == 1)
@@ -539,7 +579,12 @@ FReply SFolderHierarchyPanel::OnAddFolderToTemplate()
 FReply SFolderHierarchyPanel::OnCreateTemplateFolders()
 {
 	// Walk the tree and create all template folders in the Content Browser that are marked as MissingCreate
-	CreateFolderFromTemplateRecursive(RootItem);
+	CreateFolderFromTemplateRecursive(ContentRootItem);
+
+	for (const TSharedPtr<FTemplateFolderTreeItem>& PluginItem : PluginRootItem->Children)
+	{
+		CreateFolderFromTemplateRecursive(PluginItem);
+	}
 
 	const UProductionSettings* const ProductionSettings = GetMutableDefault<UProductionSettings>();
 	TOptional<const FCinematicProduction> ActiveProduction = ProductionSettings->GetActiveProduction();
@@ -550,7 +595,8 @@ FReply SFolderHierarchyPanel::OnCreateTemplateFolders()
 	}
 
 	TreeView->RequestTreeRefresh();
-	SetExpansionRecursive(RootItem, true);
+	SetExpansionRecursive(ContentRootItem, true);
+	SetExpansionRecursive(PluginRootItem, true);
 
 	UE::CineAssemblyToolsAnalytics::RecordEvent_ProductionCreateTemplateFolders();
 
@@ -588,7 +634,7 @@ bool SFolderHierarchyPanel::IsValidFolderName(const FText& InText, FText& OutErr
 
 	// Check for duplicate folder names
 	const FString ParentPath = FPaths::GetPath(TreeItem->Path.GetInternalPathString());
-	TSharedPtr<FTemplateFolderTreeItem> FoundItem = FindItemAtPathRecursive(RootItem, ParentPath / InText.ToString());
+	TSharedPtr<FTemplateFolderTreeItem> FoundItem = FindItemAtPath(ParentPath / InText.ToString());
 
 	if (FoundItem && FoundItem != TreeItem)
 	{
@@ -731,6 +777,26 @@ TSharedPtr<FTemplateFolderTreeItem> SFolderHierarchyPanel::FindItemAtPathRecursi
 	return nullptr;
 }
 
+TSharedPtr<FTemplateFolderTreeItem> SFolderHierarchyPanel::FindItemAtPath(const FString& Path)
+{
+	// Check for the item within the Content tree
+	if (TSharedPtr<FTemplateFolderTreeItem> FoundItem = FindItemAtPathRecursive(ContentRootItem, Path))
+	{
+		return FoundItem;
+	}
+
+	// Check for the item with the Plugins tree
+	for (const TSharedPtr<FTemplateFolderTreeItem>& PluginItem : PluginRootItem->Children)
+	{
+		if (TSharedPtr<FTemplateFolderTreeItem> FoundItem = FindItemAtPathRecursive(PluginItem, Path))
+		{
+			return FoundItem;
+		}
+	}
+
+	return nullptr;
+}
+
 void SFolderHierarchyPanel::CreateFolderFromTemplateRecursive(TSharedPtr<FTemplateFolderTreeItem> TreeItem)
 {
 	if (TreeItem->Status == ETemplateFolderStatus::MissingCreate)
@@ -867,11 +933,12 @@ FString SFolderHierarchyPanel::CreateUniqueFolderName(FContentBrowserItemPath In
 		// Check if the path matches a folder that already exists in the content browser
 		const FContentBrowserItem ExistingFolder = ContentBrowserDataSubsystem->GetItemAtPath(CombinedPathName, EContentBrowserItemTypeFilter::IncludeFolders);
 
-		// The path used by the Content Browser subsystem function was a virtual path, but our function operates on internal paths, so we have to strip "/All" off the front
+		// The path used by the Content Browser subsystem function was a virtual path, but our function operates on internal paths, so we have to strip the beginning off the front
+		CombinedPathNameStr.RemoveFromStart(TEXT("/All/Plugins"));
 		CombinedPathNameStr.RemoveFromStart(TEXT("/All"));
 
 		// Check if the path matches a folder in our template
-		TSharedPtr<FTemplateFolderTreeItem> Item = FindItemAtPathRecursive(RootItem, CombinedPathNameStr);
+		TSharedPtr<FTemplateFolderTreeItem> Item = FindItemAtPath(CombinedPathNameStr);
 
 		if (!ExistingFolder.IsValid() && !Item.IsValid())
 		{
diff --git a/Engine/Plugins/VirtualProduction/CinematicAssemblyTools/Source/CineAssemblyToolsEditor/Private/UI/SFolderHierarchyPanel.h b/Engine/Plugins/VirtualProduction/CinematicAssemblyTools/Source/CineAssemblyToolsEditor/Private/UI/SFolderHierarchyPanel.h
--- a/Engine/Plugins/VirtualProduction/CinematicAssemblyTools/Source/CineAssemblyToolsEditor/Private/UI/SFolderHierarchyPanel.h
+++ b/Engine/Plugins/VirtualProduction/CinematicAssemblyTools/Source/CineAssemblyToolsEditor/Private/UI/SFolderHierarchyPanel.h
@@ -31,6 +31,9 @@ struct FTemplateFolderTreeItem
 	/** The path of this template folder */
 	FContentBrowserItemPath Path;
 
+	/** Friendly name that should appear in the tree view. If empty, the path leaf will be displayed */
+	FString FriendlyName;
+
 	/** The parent of this item in the tree */
 	TSharedPtr<FTemplateFolderTreeItem> Parent;
 
@@ -42,6 +45,9 @@ struct FTemplateFolderTreeItem
 
 	/** The text widget that display this item's name and supports renaming */
 	TSharedPtr<SInlineEditableTextBlock> NameWidget;
+
+	/** If true, this item cannot be renamed or deleted */
+	bool bIsReadOnly = false;
 };
 
 /**
@@ -77,6 +83,9 @@ private:
 	/** Handles key presses on the tree view */
 	FReply OnKeyDown(const FGeometry& MyGeometry, const FKeyEvent& InKeyEvent);
 
+	/** Handles double clicks on the tree view to rename the item */
+	void OnTreeViewDoubleClick(TSharedPtr<FTemplateFolderTreeItem> TreeItem);
+
 	/** Builds the template folder list items (displayed in the tree view) using the current active production's folder template */
 	void BuildTreeFromProductionTemplate();
 
@@ -105,6 +114,7 @@ private:
 	void DeleteTemplateFolder(TSharedPtr<FTemplateFolderTreeItem> SelectedTreeItem);
 
 	/** Returns the template folder tree item whose path matches the input path */
+	TSharedPtr<FTemplateFolderTreeItem> FindItemAtPath(const FString& Path);
 	TSharedPtr<FTemplateFolderTreeItem> FindItemAtPathRecursive(TSharedPtr<FTemplateFolderTreeItem> TreeItem, const FString& Path);
 
 	/** Walks the template folder tree and creates all folders marked as MissingCreate */
@@ -135,17 +145,17 @@ private:
 	bool DoesPathExist(const FString& Path) const;
 
 private:
-	/* Root path of the template folder tree */
-	static FName RootPath;
-
 	/** The source template folder items for the tree view */
 	TArray<TSharedPtr<FTemplateFolderTreeItem>> FolderItemsSource;
 
 	/** The tree view of template folders for the current active production */
 	TSharedPtr<STreeView<TSharedPtr<FTemplateFolderTreeItem>>> TreeView;
 
-	/** Convenient alias for the root item (first item) in the tree view */
-	TSharedPtr<FTemplateFolderTreeItem> RootItem;
+	/** Convenient alias for the Project Content root item in the tree view */
+	TSharedPtr<FTemplateFolderTreeItem> ContentRootItem;
+
+	/** Convenient alias for the Project Plugin root item in the tree view */
+	TSharedPtr<FTemplateFolderTreeItem> PluginRootItem;
 
 	/** The cached state of the template, used to reset user changes */
 	TArray<FFolderTemplate> CachedInitialState;
