ShotManagement: Add drag/drop support to items in the template folder tree in the Production Wizard.

#rnx

diff --git a/Engine/Plugins/VirtualProduction/CinematicAssemblyTools/Source/CineAssemblyToolsEditor/Private/UI/SFolderHierarchyPanel.cpp b/Engine/Plugins/VirtualProduction/CinematicAssemblyTools/Source/CineAssemblyToolsEditor/Private/UI/SFolderHierarchyPanel.cpp
--- a/Engine/Plugins/VirtualProduction/CinematicAssemblyTools/Source/CineAssemblyToolsEditor/Private/UI/SFolderHierarchyPanel.cpp
+++ b/Engine/Plugins/VirtualProduction/CinematicAssemblyTools/Source/CineAssemblyToolsEditor/Private/UI/SFolderHierarchyPanel.cpp
@@ -2,6 +2,7 @@
 
 #include "SFolderHierarchyPanel.h"
 
+#include "Algo/Contains.h"
 #include "AssetRegistry/AssetRegistryModule.h"
 #include "AssetViewUtils.h"
 #include "CineAssemblyNamingTokens.h"
@@ -271,6 +272,12 @@ TSharedRef<ITableRow> SFolderHierarchyPanel::OnGenerateRow(TSharedPtr<FTemplateF
 	return SNew(STableRow<TSharedPtr<FTemplateFolderTreeItem>>, OwnerTable)
 		.ShowSelection(true)
 		.Padding(FMargin(8.0f, 2.0f, 8.0f, 0.0f))
+ 		.OnCanAcceptDrop_Lambda([this](const FDragDropEvent& InDragDropEvent, EItemDropZone InDropZone, TSharedPtr<FTemplateFolderTreeItem> InItem)
+ 			{
+ 				return (InItem != PluginRootItem) ? EItemDropZone::OntoItem : TOptional<EItemDropZone>();
+			})
+		.OnAcceptDrop(this, &SFolderHierarchyPanel::OnTreeRowAcceptDrop)
+		.OnDragDetected(this, &SFolderHierarchyPanel::OnTreeRowDragDetected)
 		.Content()
 		[
 			SNew(SHorizontalBox)
@@ -446,6 +453,71 @@ void SFolderHierarchyPanel::OnTreeViewDoubleClick(TSharedPtr<FTemplateFolderTree
 	}
 }
 
+FReply SFolderHierarchyPanel::OnTreeRowDragDetected(const FGeometry& MyGeometry, const FPointerEvent& MouseEvent)
+{
+	if (MouseEvent.IsMouseButtonDown(EKeys::LeftMouseButton))
+	{
+		TArray<TSharedPtr<FTemplateFolderTreeItem>> SelectedTreeItems = TreeView->GetSelectedItems();
+		if (SelectedTreeItems.Num() == 1)
+		{
+			TSharedPtr<FTemplateFolderTreeItem> SelectedTreeItem = SelectedTreeItems[0];
+
+			// Only editable tree rows can be dragged
+			if (SelectedTreeItem->bIsReadOnly)
+			{
+				return FReply::Handled();
+			}
+
+			const FString SelectedFolderName = FPaths::GetPathLeaf(SelectedTreeItem->Path.GetInternalPathString());
+			TSharedRef<FTemplateFolderDragDrop> Operation = FTemplateFolderDragDrop::New(SelectedFolderName);
+			Operation->SourceTreeItem = SelectedTreeItem;
+
+ 			Operation->OnDropNotHandled.BindLambda([this, SelectedTreeItem]()
+ 				{
+					// Return the dragged folder to its original location in the tree
+					SelectedTreeItem->Parent->Children.Add(SelectedTreeItem);
+					SetExpansionRecursive(SelectedTreeItem, true);
+					TreeView->RequestTreeRefresh();
+ 				});
+
+			// Remove this item from the tree while it is being dragged
+			SelectedTreeItem->Parent->Children.Remove(SelectedTreeItem);
+
+			TreeView->RequestTreeRefresh();
+			return FReply::Handled().BeginDragDrop(Operation);
+		}
+	}
+
+	return FReply::Unhandled();
+}
+
+FReply SFolderHierarchyPanel::OnTreeRowAcceptDrop(const FDragDropEvent& InDragDropEvent, EItemDropZone InDropZone, TSharedPtr<FTemplateFolderTreeItem> InItem)
+{
+	TSharedPtr<FTemplateFolderDragDrop> Operation = InDragDropEvent.GetOperationAs<FTemplateFolderDragDrop>();
+	if (Operation.IsValid())
+	{
+		// If an item with the same name already exists in the drop location, do not handle the drop (the dragged item will be reset to its original location)
+		if (Algo::ContainsBy(InItem->Children, FContentBrowserItemPath(InItem->Path.GetInternalPathString() / Operation->ItemName, EContentBrowserPathType::Internal), &FTemplateFolderTreeItem::Path))
+		{
+			return FReply::Unhandled();
+		}
+
+		// Add the dropped item as a child to the folder it was dropped on
+		InItem->Children.Add(Operation->SourceTreeItem);
+		Operation->SourceTreeItem->Parent = InItem;
+
+		// Update the paths of everything in the subtree to be relative to the new parent folder
+		SetChildrenPathRecursive(InItem, InItem->Path.GetInternalPathString());
+
+		InItem->Children.Sort([](const TSharedPtr<FTemplateFolderTreeItem>& A, const TSharedPtr<FTemplateFolderTreeItem>& B) { return A->Path.GetInternalPathString() < B->Path.GetInternalPathString(); });
+
+		SetExpansionRecursive(InItem, true);
+		TreeView->RequestTreeRefresh();
+	}
+
+	return FReply::Handled();
+}
+
 void SFolderHierarchyPanel::BuildTreeFromProductionTemplate()
 {
 	FolderItemsSource.Empty();
@@ -1017,4 +1089,57 @@ FText SFolderHierarchyPanel::GetFolderStatusText(TSharedPtr<FTemplateFolderTreeI
 	}
 }
 
+TSharedRef<FTemplateFolderDragDrop> FTemplateFolderDragDrop::New(const FString& Name)
+{
+	TSharedRef<FTemplateFolderDragDrop> DragDropOp = MakeShared<FTemplateFolderDragDrop>();
+	DragDropOp->MouseCursor = EMouseCursor::GrabHandClosed;
+	DragDropOp->ItemName = Name;
+	DragDropOp->Construct();
+
+	return DragDropOp;
+}
+
+TSharedPtr<SWidget> FTemplateFolderDragDrop::GetDefaultDecorator() const
+{
+	return SNew(SBorder)
+		.Padding(8.0f)
+		.BorderImage(FCineAssemblyToolsStyle::Get().GetBrush("ProductionWizard.PanelBackground"))
+		[
+			SNew(SHorizontalBox)
+
+			+ SHorizontalBox::Slot()
+				.AutoWidth()
+				.VAlign(VAlign_Center)
+				.Padding(0.0f, 0.0f, 2.0f, 0.0f)
+				[
+					SNew(SImage)
+						.ColorAndOpacity(FAppStyle::Get().GetSlateColor("ContentBrowser.DefaultFolderColor"))
+						.Image(FCineAssemblyToolsStyle::Get().GetBrush("Icons.Folder"))
+				]
+
+			+ SHorizontalBox::Slot()
+				.FillWidth(1.0f)
+				.VAlign(VAlign_Center)
+				[
+					SNew(STextBlock).Text(FText::FromString(ItemName))
+				]
+		];
+}
+
+void FTemplateFolderDragDrop::OnDragged(const class FDragDropEvent& DragDropEvent)
+{
+	if (CursorDecoratorWindow.IsValid())
+	{
+		CursorDecoratorWindow->MoveWindowTo(DragDropEvent.GetScreenSpacePosition() - (CursorDecoratorWindow->GetSizeInScreen() * FVector2f(0.0f, 0.5f)));
+	}
+}
+
+void FTemplateFolderDragDrop::OnDrop(bool bDropWasHandled, const FPointerEvent& MouseEvent)
+{
+	if (!bDropWasHandled)
+	{
+		OnDropNotHandled.ExecuteIfBound();
+	}
+}
+
 #undef LOCTEXT_NAMESPACE
diff --git a/Engine/Plugins/VirtualProduction/CinematicAssemblyTools/Source/CineAssemblyToolsEditor/Private/UI/SFolderHierarchyPanel.h b/Engine/Plugins/VirtualProduction/CinematicAssemblyTools/Source/CineAssemblyToolsEditor/Private/UI/SFolderHierarchyPanel.h
--- a/Engine/Plugins/VirtualProduction/CinematicAssemblyTools/Source/CineAssemblyToolsEditor/Private/UI/SFolderHierarchyPanel.h
+++ b/Engine/Plugins/VirtualProduction/CinematicAssemblyTools/Source/CineAssemblyToolsEditor/Private/UI/SFolderHierarchyPanel.h
@@ -50,6 +50,35 @@ struct FTemplateFolderTreeItem
 	bool bIsReadOnly = false;
 };
 
+/** Drag and Drop operation to handle dragging template folders  */
+class FTemplateFolderDragDrop : public FDragDropOperation
+{
+public:
+
+	DRAG_DROP_OPERATOR_TYPE(FTemplateFolderDragDrop, FDragDropOperation)
+
+	DECLARE_DELEGATE(FOnDropNotHandled);
+
+	FTemplateFolderDragDrop() = default;
+
+	static TSharedRef<FTemplateFolderDragDrop> New(const FString& Name);
+
+	// Begin FDragDropOperation overrides
+	virtual void OnDrop(bool bDropWasHandled, const FPointerEvent& MouseEvent) override;
+	virtual void OnDragged(const class FDragDropEvent& DragDropEvent) override;
+	virtual TSharedPtr<SWidget> GetDefaultDecorator() const override;
+	// End FDragDropOperation overrides
+
+	/** The source item being dragged. This could be null if the thing being dragged is from the Asset List (not the tree) */
+	TSharedPtr<FTemplateFolderTreeItem> SourceTreeItem;
+
+	/** The name of the item being displayed, used by the decorator and to name the new tree item that will get made when dropped */
+	FString ItemName;
+
+	/** Callback to handle the case where an item is dropped somewhere not on the tree view */
+	FOnDropNotHandled OnDropNotHandled;
+};
+
 /**
  * UI for the Folder Hierarchy panel in the Production Wizard
  */
@@ -86,6 +115,12 @@ private:
 	/** Handles double clicks on the tree view to rename the item */
 	void OnTreeViewDoubleClick(TSharedPtr<FTemplateFolderTreeItem> TreeItem);
 
+	/** Begins a drag and drop event to drag an item out of the content tree view */
+	FReply OnTreeRowDragDetected(const FGeometry& MyGeometry, const FPointerEvent& MouseEvent);
+
+	/** Callback when an item is dropped onto a folder in the content tree view to add it to the children of that row */
+	FReply OnTreeRowAcceptDrop(const FDragDropEvent& InDragDropEvent, EItemDropZone InDropZone, TSharedPtr<FTemplateFolderTreeItem> InItem);
+
 	/** Builds the template folder list items (displayed in the tree view) using the current active production's folder template */
 	void BuildTreeFromProductionTemplate();
 
