ShotManagement: Add an edit condition to the Take Recorder Assembly Schema property to hide the property when the TargetRecordClass is not CineAssembly

#rnx

diff --git a/Engine/Plugins/VirtualProduction/CinematicAssemblyTools/Source/CineAssemblyToolsEditor/Private/TakeRecorder/CineAssemblyTakeRecorderSettings.cpp b/Engine/Plugins/VirtualProduction/CinematicAssemblyTools/Source/CineAssemblyToolsEditor/Private/TakeRecorder/CineAssemblyTakeRecorderSettings.cpp
--- a/Engine/Plugins/VirtualProduction/CinematicAssemblyTools/Source/CineAssemblyToolsEditor/Private/TakeRecorder/CineAssemblyTakeRecorderSettings.cpp
+++ b/Engine/Plugins/VirtualProduction/CinematicAssemblyTools/Source/CineAssemblyToolsEditor/Private/TakeRecorder/CineAssemblyTakeRecorderSettings.cpp
@@ -2,8 +2,16 @@
 
 #include "CineAssemblyTakeRecorderSettings.h"
 
+#include "CineAssembly.h"
+#include "TakePresetSettings.h"
+
 void UCineAssemblyTakeRecorderSettings::PostEditChangeProperty(FPropertyChangedEvent& PropertyChangedEvent)
 {
 	UObject::PostEditChangeProperty(PropertyChangedEvent);
 	SaveConfig();
 }
+
+bool UCineAssemblyTakeRecorderSettings::CanEditAssemblySchema() const
+{
+	return UTakePresetSettings::Get()->GetTargetRecordClass() == UCineAssembly::StaticClass();
+}
diff --git a/Engine/Plugins/VirtualProduction/CinematicAssemblyTools/Source/CineAssemblyToolsEditor/Private/TakeRecorder/CineAssemblyTakeRecorderSettings.h b/Engine/Plugins/VirtualProduction/CinematicAssemblyTools/Source/CineAssemblyToolsEditor/Private/TakeRecorder/CineAssemblyTakeRecorderSettings.h
--- a/Engine/Plugins/VirtualProduction/CinematicAssemblyTools/Source/CineAssemblyToolsEditor/Private/TakeRecorder/CineAssemblyTakeRecorderSettings.h
+++ b/Engine/Plugins/VirtualProduction/CinematicAssemblyTools/Source/CineAssemblyToolsEditor/Private/TakeRecorder/CineAssemblyTakeRecorderSettings.h
@@ -17,8 +17,12 @@ public:
 	virtual void PostEditChangeProperty(FPropertyChangedEvent& PropertyChangedEvent) override;
 	//~ End UObject Interface
 
+	/** Returns true if the Take Preset settings TargetRecordClass is set to CineAssembly */
+	UFUNCTION()
+	bool CanEditAssemblySchema() const;
+
 public:
 	/** The Cine Assembly Schema to use as the base for the recorded Assembly */
-	UPROPERTY(Config, EditAnywhere, Category = "Cine Assembly")
+	UPROPERTY(Config, EditAnywhere, Category = "Cine Assembly", meta=(EditCondition = "CanEditAssemblySchema", EditConditionHides))
 	TSoftObjectPtr<UCineAssemblySchema> AssemblySchema;
 };
