ShotManagement: Fix integer metadata to support leading zeros
#rnx

diff --git a/Engine/Plugins/VirtualProduction/CinematicAssemblyTools/Source/CineAssemblyToolsEditor/Private/CineAssemblyCustomization.cpp b/Engine/Plugins/VirtualProduction/CinematicAssemblyTools/Source/CineAssemblyToolsEditor/Private/CineAssemblyCustomization.cpp
--- a/Engine/Plugins/VirtualProduction/CinematicAssemblyTools/Source/CineAssemblyToolsEditor/Private/CineAssemblyCustomization.cpp
+++ b/Engine/Plugins/VirtualProduction/CinematicAssemblyTools/Source/CineAssemblyToolsEditor/Private/CineAssemblyCustomization.cpp
@@ -41,6 +41,8 @@ void FCineAssemblyCustomization::CustomizeDetails(IDetailLayoutBuilder& DetailBu
 
 	CustomizedCineAssembly = Cast<UCineAssembly>(CustomizedObjects[0]);
 
+	LeadingZeroTypeInterface = MakeShared<FLeadingZeroNumericTypeInterface>();
+
 	CustomizeDefaultCategory(DetailBuilder);
 	CustomizeMetadataCategory(DetailBuilder);
 	CustomizeSubsequenceCategory(DetailBuilder);
@@ -166,6 +168,7 @@ void FCineAssemblyCustomization::CustomizeMetadataCategory(IDetailLayoutBuilder&
 		{
 			ValueWidget = SNew(SNumericEntryBox<int32>)
 				.Font(FAppStyle::GetFontStyle(TEXT("PropertyWindow.NormalFont")))
+				.TypeInterface(LeadingZeroTypeInterface)
 				.Value_Lambda([this, MetadataDesc]()
 					{
 						int32 Value;
@@ -177,7 +180,10 @@ void FCineAssemblyCustomization::CustomizeMetadataCategory(IDetailLayoutBuilder&
 					})
 				.OnValueChanged_Lambda([this, MetadataDesc](int32 InValue)
 					{
-						CustomizedCineAssembly->SetMetadataAsInteger(MetadataDesc.Key, InValue);
+						// In order to preserve the desired number of leading zeroes, the integer value is converted to a string using the type interface
+						// and then stored in the assembly as string metadata.
+						const FString ValueString = LeadingZeroTypeInterface->ToString(InValue);
+						CustomizedCineAssembly->SetMetadataAsString(MetadataDesc.Key, ValueString);
 					});
 		}
 		else if (MetadataDesc.Type == ECineAssemblyMetadataType::Float)
diff --git a/Engine/Plugins/VirtualProduction/CinematicAssemblyTools/Source/CineAssemblyToolsEditor/Private/CineAssemblyCustomization.h b/Engine/Plugins/VirtualProduction/CinematicAssemblyTools/Source/CineAssemblyToolsEditor/Private/CineAssemblyCustomization.h
--- a/Engine/Plugins/VirtualProduction/CinematicAssemblyTools/Source/CineAssemblyToolsEditor/Private/CineAssemblyCustomization.h
+++ b/Engine/Plugins/VirtualProduction/CinematicAssemblyTools/Source/CineAssemblyToolsEditor/Private/CineAssemblyCustomization.h
@@ -4,6 +4,7 @@
 
 #include "IDetailCustomization.h"
 
+#include "LeadingZeroNumericTypeInterface.h"
 #include "Styling/SlateTypes.h"
 #include "UObject/TemplateString.h"
 
@@ -77,4 +78,7 @@ private:
 
 	/** The last time the naming tokens were updated */
 	FDateTime LastTokenUpdateTime;
+
+	/** Type interface for a numeric entry box that supports leading zeroes */
+	TSharedPtr<FLeadingZeroNumericTypeInterface> LeadingZeroTypeInterface;
 };
diff --git a/Engine/Plugins/VirtualProduction/CinematicAssemblyTools/Source/CineAssemblyToolsEditor/Private/CineAssemblyMetadataCustomization.cpp b/Engine/Plugins/VirtualProduction/CinematicAssemblyTools/Source/CineAssemblyToolsEditor/Private/CineAssemblyMetadataCustomization.cpp
--- a/Engine/Plugins/VirtualProduction/CinematicAssemblyTools/Source/CineAssemblyToolsEditor/Private/CineAssemblyMetadataCustomization.cpp
+++ b/Engine/Plugins/VirtualProduction/CinematicAssemblyTools/Source/CineAssemblyToolsEditor/Private/CineAssemblyMetadataCustomization.cpp
@@ -116,6 +116,8 @@ void FCineAssemblyMetadataCustomization::CustomizeChildren(TSharedRef<IPropertyH
 
 	Schema = Cast<UCineAssemblySchema>(CustomizedObjects[0].Get());
 
+	LeadingZeroTypeInterface = MakeShared<FLeadingZeroNumericTypeInterface>();
+
 	if (!PropertyHandle->IsExpanded())
 	{
 		PropertyHandle->SetExpanded(true);
@@ -209,6 +211,7 @@ void FCineAssemblyMetadataCustomization::CustomizeChildren(TSharedRef<IPropertyH
 			[
 				SNew(SNumericEntryBox<int32>)
 					.Font(FAppStyle::GetFontStyle(TEXT("PropertyWindow.NormalFont")))
+					.TypeInterface(LeadingZeroTypeInterface)
 					.Value_Lambda([&MetadataDesc]() -> int32
 						{
 							if (const int32* Value = MetadataDesc.DefaultValue.TryGet<int32>())
diff --git a/Engine/Plugins/VirtualProduction/CinematicAssemblyTools/Source/CineAssemblyToolsEditor/Private/CineAssemblyMetadataCustomization.h b/Engine/Plugins/VirtualProduction/CinematicAssemblyTools/Source/CineAssemblyToolsEditor/Private/CineAssemblyMetadataCustomization.h
--- a/Engine/Plugins/VirtualProduction/CinematicAssemblyTools/Source/CineAssemblyToolsEditor/Private/CineAssemblyMetadataCustomization.h
+++ b/Engine/Plugins/VirtualProduction/CinematicAssemblyTools/Source/CineAssemblyToolsEditor/Private/CineAssemblyMetadataCustomization.h
@@ -5,6 +5,7 @@
 #include "IPropertyTypeCustomization.h"
 
 #include "CineAssemblySchema.h"
+#include "LeadingZeroNumericTypeInterface.h"
 
 class SWidget;
 
@@ -40,4 +41,7 @@ private:
 
 	/** The array index of the metadata struct being customized */
 	int32 ArrayIndex;
+
+	/** Type interface for a numeric entry box that supports leading zeroes */
+	TSharedPtr<FLeadingZeroNumericTypeInterface> LeadingZeroTypeInterface;
 };
diff --git a/Engine/Plugins/VirtualProduction/CinematicAssemblyTools/Source/CineAssemblyToolsEditor/Private/LeadingZeroNumericTypeInterface.cpp b/Engine/Plugins/VirtualProduction/CinematicAssemblyTools/Source/CineAssemblyToolsEditor/Private/LeadingZeroNumericTypeInterface.cpp
new file mode 100644
--- /dev/null
+++ b/Engine/Plugins/VirtualProduction/CinematicAssemblyTools/Source/CineAssemblyToolsEditor/Private/LeadingZeroNumericTypeInterface.cpp
@@ -0,0 +1,26 @@
+// Copyright Epic Games, Inc. All Rights Reserved.
+
+#include "LeadingZeroNumericTypeInterface.h"
+
+FString FLeadingZeroNumericTypeInterface::ToString(const int32& Value) const
+{
+	const FNumberFormattingOptions NumberFormattingOptions = FNumberFormattingOptions()
+		.SetUseGrouping(false)
+		.SetMinimumIntegralDigits(MinimumIntegralDigits);
+
+	return FastDecimalFormat::NumberToString(Value, ExpressionParser::GetLocalizedNumberFormattingRules(), NumberFormattingOptions);
+}
+
+TOptional<int32> FLeadingZeroNumericTypeInterface::FromString(const FString& InString, const int32& InExistingValue)
+{
+	// Convert the string to an integer using the base default type interface
+	const TOptional<int32> Value = TDefaultNumericTypeInterface<int32>::FromString(InString, InExistingValue);
+
+	// If the value was successfully parsed, update the minimum number of digits to match the length of the input string
+	if (Value.IsSet())
+	{
+		MinimumIntegralDigits = InString.Len();
+	}
+
+	return Value;
+}
diff --git a/Engine/Plugins/VirtualProduction/CinematicAssemblyTools/Source/CineAssemblyToolsEditor/Private/LeadingZeroNumericTypeInterface.h b/Engine/Plugins/VirtualProduction/CinematicAssemblyTools/Source/CineAssemblyToolsEditor/Private/LeadingZeroNumericTypeInterface.h
new file mode 100644
--- /dev/null
+++ b/Engine/Plugins/VirtualProduction/CinematicAssemblyTools/Source/CineAssemblyToolsEditor/Private/LeadingZeroNumericTypeInterface.h
@@ -0,0 +1,18 @@
+// Copyright Epic Games, Inc. All Rights Reserved.
+
+#pragma once
+
+#include "Widgets/Input/NumericTypeInterface.h"
+
+struct FLeadingZeroNumericTypeInterface : public TDefaultNumericTypeInterface<int32>
+{
+	FLeadingZeroNumericTypeInterface() = default;
+
+	//~ Begin INumericTypeInterface overrides
+	virtual FString ToString(const int32& Value) const override;
+	virtual TOptional<int32> FromString(const FString& InString, const int32& InExistingValue) override;
+	//~ End INumericTypeInterface overrides
+
+private:
+	int32 MinimumIntegralDigits = 1;
+};
\ No newline at end of file
