ShotManagement: Add missing unregister of a property type customization
#rnx

diff --git a/Engine/Plugins/VirtualProduction/CinematicAssemblyTools/Source/CineAssemblyToolsEditor/Private/CineAssemblyToolsEditorModule.cpp b/Engine/Plugins/VirtualProduction/CinematicAssemblyTools/Source/CineAssemblyToolsEditor/Private/CineAssemblyToolsEditorModule.cpp
--- a/Engine/Plugins/VirtualProduction/CinematicAssemblyTools/Source/CineAssemblyToolsEditor/Private/CineAssemblyToolsEditorModule.cpp
+++ b/Engine/Plugins/VirtualProduction/CinematicAssemblyTools/Source/CineAssemblyToolsEditor/Private/CineAssemblyToolsEditorModule.cpp
@@ -198,6 +198,7 @@ void FCineAssemblyToolsEditorModule::ShutdownModule()
 	PropertyModule.UnregisterCustomClassLayout(UProductionSettings::StaticClass()->GetFName());
 	PropertyModule.UnregisterCustomClassLayout(UCineAssembly::StaticClass()->GetFName());
 	PropertyModule.UnregisterCustomClassLayout(UCineAssemblySchema::StaticClass()->GetFName());
+	PropertyModule.UnregisterCustomPropertyTypeLayout(FCinematicProduction::StaticStruct()->GetFName());
 	PropertyModule.UnregisterCustomPropertyTypeLayout(FAssemblyMetadataDesc::StaticStruct()->GetFName());
 
 	FGlobalTabmanager::Get()->UnregisterNomadTabSpawner(ProductionWizardTabName);
