ShotManagement: Fix Cine Assemblies not automatically adding metadata tokens when setting the schema.
#rnx

diff --git a/Engine/Plugins/VirtualProduction/CinematicAssemblyTools/Source/CineAssemblyTools/Private/CineAssembly.cpp b/Engine/Plugins/VirtualProduction/CinematicAssemblyTools/Source/CineAssemblyTools/Private/CineAssembly.cpp
--- a/Engine/Plugins/VirtualProduction/CinematicAssemblyTools/Source/CineAssemblyTools/Private/CineAssembly.cpp
+++ b/Engine/Plugins/VirtualProduction/CinematicAssemblyTools/Source/CineAssemblyTools/Private/CineAssembly.cpp
@@ -217,22 +217,22 @@ void UCineAssembly::ChangeSchema(UCineAssemblySchema* InSchema)
 			if (MetadataDesc.DefaultValue.IsType<FString>())
 			{
 				const FString& DefaultValue = MetadataDesc.DefaultValue.Get<FString>();
-				MetadataJsonObject->SetStringField(MetadataDesc.Key, DefaultValue);
+				SetMetadataAsString(MetadataDesc.Key, DefaultValue);
 			}
 			else if (MetadataDesc.DefaultValue.IsType<bool>())
 			{
 				const bool& DefaultValue = MetadataDesc.DefaultValue.Get<bool>();
-				MetadataJsonObject->SetBoolField(MetadataDesc.Key, DefaultValue);
+				SetMetadataAsBool(MetadataDesc.Key, DefaultValue);
 			}
 			else if (MetadataDesc.DefaultValue.IsType<int32>())
 			{
 				const int32& DefaultValue = MetadataDesc.DefaultValue.Get<int32>();
-				MetadataJsonObject->SetNumberField(MetadataDesc.Key, DefaultValue);
+				SetMetadataAsInteger(MetadataDesc.Key, DefaultValue);
 			}
 			else if (MetadataDesc.DefaultValue.IsType<float>())
 			{
 				const float& DefaultValue = MetadataDesc.DefaultValue.Get<float>();
-				MetadataJsonObject->SetNumberField(MetadataDesc.Key, DefaultValue);
+				SetMetadataAsFloat(MetadataDesc.Key, DefaultValue);
 			}
 		}
 	}
diff --git a/Engine/Plugins/VirtualProduction/CinematicAssemblyTools/Source/CineAssemblyToolsEditor/Private/CineAssemblyMetadataCustomization.cpp b/Engine/Plugins/VirtualProduction/CinematicAssemblyTools/Source/CineAssemblyToolsEditor/Private/CineAssemblyMetadataCustomization.cpp
--- a/Engine/Plugins/VirtualProduction/CinematicAssemblyTools/Source/CineAssemblyToolsEditor/Private/CineAssemblyMetadataCustomization.cpp
+++ b/Engine/Plugins/VirtualProduction/CinematicAssemblyTools/Source/CineAssemblyToolsEditor/Private/CineAssemblyMetadataCustomization.cpp
@@ -88,8 +88,12 @@ TSharedRef<SWidget> FCineAssemblyMetadataCustomization::MakeStringDefaultValueWi
 					{
 						UNamingTokensEngineSubsystem* NamingTokensSubsystem = GEngine->GetEngineSubsystem<UNamingTokensEngineSubsystem>();
 						FNamingTokenFilterArgs FilterArgs;
-						FNamingTokenResultData Result = NamingTokensSubsystem->EvaluateTokenString(MetadataDesc.DefaultValue.Get<FString>(), FilterArgs, {});
-						return Result.EvaluatedText;
+						if (const FString* DefaultValue = MetadataDesc.DefaultValue.TryGet<FString>())
+						{
+							FNamingTokenResultData Result = NamingTokensSubsystem->EvaluateTokenString(*DefaultValue, FilterArgs, {});
+							return Result.EvaluatedText;
+						}
+						return FText::GetEmpty();
 					})
 				.OnTextCommitted_Lambda([&MetadataDesc](const FText& InText, ETextCommit::Type InCommitType)
 					{
