ShotManagement: Fix schema and assembly windows being minimized when the user focuses on another UE window.
#rnx

diff --git a/Engine/Plugins/VirtualProduction/CinematicAssemblyTools/Source/CineAssemblyToolsEditor/Private/CineAssemblyFactory.cpp b/Engine/Plugins/VirtualProduction/CinematicAssemblyTools/Source/CineAssemblyToolsEditor/Private/CineAssemblyFactory.cpp
--- a/Engine/Plugins/VirtualProduction/CinematicAssemblyTools/Source/CineAssemblyToolsEditor/Private/CineAssemblyFactory.cpp
+++ b/Engine/Plugins/VirtualProduction/CinematicAssemblyTools/Source/CineAssemblyToolsEditor/Private/CineAssemblyFactory.cpp
@@ -10,6 +10,7 @@
 #include "ContentBrowserModule.h"
 #include "Framework/Application/SlateApplication.h"
 #include "IContentBrowserSingleton.h"
+#include "Interfaces/IMainFrameModule.h"
 #include "Modules/ModuleManager.h"
 #include "MovieScene.h"
 #include "MovieSceneToolsProjectSettings.h"
@@ -35,7 +36,18 @@ bool UCineAssemblyFactory::ConfigureProperties()
 	const FString CurrentPath = ContentBrowser.GetCurrentPath().GetInternalPathString();
 
 	TSharedRef<SCineAssemblyConfigWindow> CineAssemblyConfigWindow = SNew(SCineAssemblyConfigWindow, CurrentPath);
-	FSlateApplication::Get().AddWindow(CineAssemblyConfigWindow);
+
+	const IMainFrameModule& MainFrameModule = FModuleManager::LoadModuleChecked<IMainFrameModule>(TEXT("MainFrame"));
+	const TSharedPtr<SWindow> ParentWindow = MainFrameModule.GetParentWindow();
+
+	if (ParentWindow.IsValid())
+	{
+		FSlateApplication::Get().AddWindowAsNativeChild(CineAssemblyConfigWindow, ParentWindow.ToSharedRef());
+	}
+	else
+	{
+		FSlateApplication::Get().AddWindow(CineAssemblyConfigWindow);
+	}
 
 	return false;
 }
diff --git a/Engine/Plugins/VirtualProduction/CinematicAssemblyTools/Source/CineAssemblyToolsEditor/Private/CineAssemblySchemaFactory.cpp b/Engine/Plugins/VirtualProduction/CinematicAssemblyTools/Source/CineAssemblyToolsEditor/Private/CineAssemblySchemaFactory.cpp
--- a/Engine/Plugins/VirtualProduction/CinematicAssemblyTools/Source/CineAssemblyToolsEditor/Private/CineAssemblySchemaFactory.cpp
+++ b/Engine/Plugins/VirtualProduction/CinematicAssemblyTools/Source/CineAssemblyToolsEditor/Private/CineAssemblySchemaFactory.cpp
@@ -10,6 +10,7 @@
 #include "ContentBrowserModule.h"
 #include "Framework/Application/SlateApplication.h"
 #include "IContentBrowserSingleton.h"
+#include "Interfaces/IMainFrameModule.h"
 #include "Modules/ModuleManager.h"
 #include "UI/CineAssembly/SCineAssemblySchemaWindow.h"
 
@@ -42,7 +43,18 @@ bool UCineAssemblySchemaFactory::ConfigureProperties()
 		.SupportsMinimize(false);
 
 	NewSchemaWindow->SetContent(SchemaWidget);
-	FSlateApplication::Get().AddWindow(NewSchemaWindow);
+
+	const IMainFrameModule& MainFrameModule = FModuleManager::LoadModuleChecked<IMainFrameModule>(TEXT("MainFrame"));
+	const TSharedPtr<SWindow> ParentWindow = MainFrameModule.GetParentWindow();
+
+	if (ParentWindow.IsValid())
+	{
+		FSlateApplication::Get().AddWindowAsNativeChild(NewSchemaWindow, ParentWindow.ToSharedRef());
+	}
+	else
+	{
+		FSlateApplication::Get().AddWindow(NewSchemaWindow);
+	}
 
 	return false;
 }
