ShotManagement: Fix assembly string metadata with tokens to not re-evaluate every time the metadata value is queried. This fix also updates the metadata value to store both the unresolved and resolved text to support manual re-evaluation (instead of automatic).
#rnx

diff --git a/Engine/Plugins/VirtualProduction/CinematicAssemblyTools/Source/CineAssemblyTools/CineAssemblyTools.build.cs b/Engine/Plugins/VirtualProduction/CinematicAssemblyTools/Source/CineAssemblyTools/CineAssemblyTools.build.cs
--- a/Engine/Plugins/VirtualProduction/CinematicAssemblyTools/Source/CineAssemblyTools/CineAssemblyTools.build.cs
+++ b/Engine/Plugins/VirtualProduction/CinematicAssemblyTools/Source/CineAssemblyTools/CineAssemblyTools.build.cs
@@ -21,6 +21,7 @@ namespace UnrealBuildTool.Rules
 					"Core",
 					"CoreUObject",
 					"Json",
+					"JsonUtilities",
 					"MovieScene",
 					"SlateCore",
 					"UniversalObjectLocator"
diff --git a/Engine/Plugins/VirtualProduction/CinematicAssemblyTools/Source/CineAssemblyTools/Private/CineAssembly.cpp b/Engine/Plugins/VirtualProduction/CinematicAssemblyTools/Source/CineAssemblyTools/Private/CineAssembly.cpp
--- a/Engine/Plugins/VirtualProduction/CinematicAssemblyTools/Source/CineAssemblyTools/Private/CineAssembly.cpp
+++ b/Engine/Plugins/VirtualProduction/CinematicAssemblyTools/Source/CineAssemblyTools/Private/CineAssembly.cpp
@@ -5,6 +5,7 @@
 #include "AssetRegistry/AssetRegistryModule.h"
 #include "CineAssemblyNamingTokens.h"
 #include "HAL/FileManager.h"
+#include "JsonObjectConverter.h"
 #include "LevelSequenceShotMetaDataLibrary.h"
 #include "MovieScene.h"
 #include "NamingTokensEngineSubsystem.h"
@@ -216,8 +217,18 @@ void UCineAssembly::ChangeSchema(UCineAssemblySchema* InSchema)
 		{
 			if (MetadataDesc.DefaultValue.IsType<FString>())
 			{
-				const FString& DefaultValue = MetadataDesc.DefaultValue.Get<FString>();
-				SetMetadataAsString(MetadataDesc.Key, DefaultValue);
+				if (MetadataDesc.bEvaluateTokens)
+				{
+					FTemplateString DefaultTemplateString;
+					DefaultTemplateString.Template = MetadataDesc.DefaultValue.Get<FString>();
+					DefaultTemplateString.Resolved = UCineAssemblyNamingTokens::GetResolvedText(DefaultTemplateString.Template, this);
+					SetMetadataAsTokenString(MetadataDesc.Key, DefaultTemplateString);
+				}
+				else
+				{
+					const FString& DefaultValue = MetadataDesc.DefaultValue.Get<FString>();
+					SetMetadataAsString(MetadataDesc.Key, DefaultValue);
+				}
 			}
 			else if (MetadataDesc.DefaultValue.IsType<bool>())
 			{
@@ -574,6 +585,13 @@ FString UCineAssembly::GetFullMetadataString() const
 	return JsonString;
 }
 
+TArray<FString> UCineAssembly::GetMetadataKeys() const
+{
+	TArray<FString> Keys;
+	MetadataJsonObject->Values.GetKeys(Keys);
+	return Keys;
+}
+
 void UCineAssembly::SetMetadataAsString(FString InKey, FString InValue)
 {
 	Modify();
@@ -581,6 +599,14 @@ void UCineAssembly::SetMetadataAsString(FString InKey, FString InValue)
 	AddMetadataNamingToken(InKey);
 }
 
+void UCineAssembly::SetMetadataAsTokenString(FString InKey, FTemplateString InValue)
+{
+	Modify();
+	TSharedPtr<FJsonObject> TemplateStringObject = FJsonObjectConverter::UStructToJsonObject<FTemplateString>(InValue);
+	MetadataJsonObject->SetObjectField(InKey, TemplateStringObject);
+	AddMetadataNamingToken(InKey);
+}
+
 void UCineAssembly::SetMetadataAsBool(FString InKey, bool InValue)
 {
 	Modify();
@@ -603,24 +629,32 @@ void UCineAssembly::SetMetadataAsFloat(FString InKey, float InValue)
 }
 
 bool UCineAssembly::GetMetadataAsString(FString InKey, FString& OutValue) const
-{
-	return GetMetadataAsString(InKey, OutValue, true /* bEvaluateTokens */);
-}
-
-bool UCineAssembly::GetMetadataAsString(FString InKey, FString& OutValue, bool bEvaluateTokens) const
 {
 	if (!MetadataJsonObject->TryGetStringField(InKey, OutValue))
 	{
-		OutValue = TEXT("");
-		return false;
+		FTemplateString TemplateString;
+		if (!GetMetadataAsTokenString(InKey, TemplateString))
+		{
+			OutValue = TEXT("");
+			return false;
+		}
+
+		OutValue = TemplateString.Resolved.ToString();
 	}
+	return true;
+}
 
-	if (bEvaluateTokens)
+bool UCineAssembly::GetMetadataAsTokenString(FString InKey, FTemplateString& OutValue) const
+{
+	const TSharedPtr<FJsonObject>* TemplateStringObject;
+	if (!MetadataJsonObject->TryGetObjectField(InKey, TemplateStringObject))
 	{
-		// Resolve any tokens in the returned string
-		OutValue = UCineAssemblyNamingTokens::GetResolvedText(OutValue, const_cast<UCineAssembly*>(this)).ToString();
+		OutValue.Template = TEXT("");
+		OutValue.Resolved = FText::GetEmpty();
+		return false;
 	}
 
+	FJsonObjectConverter::JsonObjectToUStruct<FTemplateString>(TemplateStringObject->ToSharedRef(), &OutValue);
 	return true;
 }
 
diff --git a/Engine/Plugins/VirtualProduction/CinematicAssemblyTools/Source/CineAssemblyTools/Private/CineAssemblyNamingTokens.cpp b/Engine/Plugins/VirtualProduction/CinematicAssemblyTools/Source/CineAssemblyTools/Private/CineAssemblyNamingTokens.cpp
--- a/Engine/Plugins/VirtualProduction/CinematicAssemblyTools/Source/CineAssemblyTools/Private/CineAssemblyNamingTokens.cpp
+++ b/Engine/Plugins/VirtualProduction/CinematicAssemblyTools/Source/CineAssemblyTools/Private/CineAssemblyNamingTokens.cpp
@@ -213,7 +213,7 @@ void UCineAssemblyNamingTokens::AddMetadataToken(const FString& InTokenKey)
 			if (InAssembly.IsValid())
 			{
 				FString ValueString;
-				if (InAssembly->GetMetadataAsString(InTokenKey, ValueString, false))
+				if (InAssembly->GetMetadataAsString(InTokenKey, ValueString))
 				{
 					// Test if the value string could be an object path, and if it is, return just the filename instead of the full path
 					FSoftObjectPath ObjectPath = FSoftObjectPath(ValueString);
diff --git a/Engine/Plugins/VirtualProduction/CinematicAssemblyTools/Source/CineAssemblyTools/Public/CineAssembly.h b/Engine/Plugins/VirtualProduction/CinematicAssemblyTools/Source/CineAssemblyTools/Public/CineAssembly.h
--- a/Engine/Plugins/VirtualProduction/CinematicAssemblyTools/Source/CineAssemblyTools/Public/CineAssembly.h
+++ b/Engine/Plugins/VirtualProduction/CinematicAssemblyTools/Source/CineAssemblyTools/Public/CineAssembly.h
@@ -126,10 +126,17 @@ public:
 	UFUNCTION(BlueprintPure, Category = "Cine Assembly Tools")
 	CINEASSEMBLYTOOLS_API FString GetFullMetadataString() const;
 
+	/** Get all of the metadata keys for this assembly */
+	CINEASSEMBLYTOOLS_API TArray<FString> GetMetadataKeys() const;
+
 	/** Add a string as metadata to this assembly */
 	UFUNCTION(BlueprintCallable, Category = "Cine Assembly Tools")
 	CINEASSEMBLYTOOLS_API void SetMetadataAsString(FString InKey, FString InValue);
 
+	/** Add a tokenized string as metadata to this assembly */
+	UFUNCTION(BlueprintCallable, Category = "Cine Assembly Tools")
+	CINEASSEMBLYTOOLS_API void SetMetadataAsTokenString(FString InKey, FTemplateString InValue);
+
 	/** Add a boolean as metadata to this assembly */
 	UFUNCTION(BlueprintCallable, Category = "Cine Assembly Tools")
 	CINEASSEMBLYTOOLS_API void SetMetadataAsBool(FString InKey, bool InValue);
@@ -142,12 +149,13 @@ public:
 	UFUNCTION(BlueprintCallable, Category = "Cine Assembly Tools")
 	CINEASSEMBLYTOOLS_API void SetMetadataAsFloat(FString InKey, float InValue);
 
-	/** Get the metadata value for the input key as a token evaluated string (if it exists) */
+	/** Get the metadata value for the input key as a string (if it exists) */
 	UFUNCTION(BlueprintPure, Category = "Cine Assembly Tools", meta = (ReturnDisplayName = "FoundKey"))
 	CINEASSEMBLYTOOLS_API bool GetMetadataAsString(FString InKey, FString& OutValue) const;
 
-	/** Get the metadata value for the input key as a string (if it exists), with tokens evaluated if requested. */
-	CINEASSEMBLYTOOLS_API bool GetMetadataAsString(FString InKey, FString& OutValue, bool bEvaluateTokens) const;
+	/** Get the metadata value for the input key as a token string (if it exists) */
+	UFUNCTION(BlueprintPure, Category = "Cine Assembly Tools", meta = (ReturnDisplayName = "FoundKey"))
+	CINEASSEMBLYTOOLS_API bool GetMetadataAsTokenString(FString InKey, FTemplateString& OutValue) const;
 
 	/** Get the metadata value for the input key as a boolean (if it exists) */
 	UFUNCTION(BlueprintPure, Category = "Cine Assembly Tools", meta = (ReturnDisplayName = "FoundKey"))
diff --git a/Engine/Plugins/VirtualProduction/CinematicAssemblyTools/Source/CineAssemblyToolsEditor/Private/CineAssemblyCustomization.cpp b/Engine/Plugins/VirtualProduction/CinematicAssemblyTools/Source/CineAssemblyToolsEditor/Private/CineAssemblyCustomization.cpp
--- a/Engine/Plugins/VirtualProduction/CinematicAssemblyTools/Source/CineAssemblyToolsEditor/Private/CineAssemblyCustomization.cpp
+++ b/Engine/Plugins/VirtualProduction/CinematicAssemblyTools/Source/CineAssemblyToolsEditor/Private/CineAssemblyCustomization.cpp
@@ -123,147 +123,183 @@ void FCineAssemblyCustomization::CustomizeMetadataCategory(IDetailLayoutBuilder&
 		return;
 	}
 
-	// Add a new category for Schema Metadata properties
-	const FText MetadataCategoryName = FText::Format(LOCTEXT("SchemaMetadataCategoryName", "{0} Metadata"), FText::FromString(BaseSchema->SchemaName));
-	IDetailCategoryBuilder& MetadataCategory = DetailBuilder.EditCategory("SchemaMetadata", MetadataCategoryName);
-
-	// Add a property row for each metadata struct in the customized Assembly's base schema
-	for (const FAssemblyMetadataDesc& MetadataDesc : BaseSchema->AssemblyMetadata)
+	// Metadata is editable for new (transient) assemblies, but read-only for existing assemblies, so we customize the details view differently
+	if (CustomizedCineAssembly->GetPackage() == GetTransientPackage())
 	{
-		if (MetadataDesc.Key.IsEmpty())
-		{
-			continue;
-		}
+		// Add a new category for Schema Metadata properties
+		const FText MetadataCategoryName = FText::Format(LOCTEXT("SchemaMetadataCategoryName", "{0} Metadata"), FText::FromString(BaseSchema->SchemaName));
+		IDetailCategoryBuilder& MetadataCategory = DetailBuilder.EditCategory("SchemaMetadata", MetadataCategoryName);
 
-		CustomizedCineAssembly->AddMetadataNamingToken(MetadataDesc.Key);
-
-		TSharedPtr<SWidget> ValueWidget;
-		if (MetadataDesc.Type == ECineAssemblyMetadataType::String)
-		{
-			ValueWidget = SNew(SBox)
-				.MaxDesiredHeight(120.0f)
-				[
-					MakeStringValueWidget(MetadataDesc)
-				];
-		}
-		else if (MetadataDesc.Type == ECineAssemblyMetadataType::Bool)
+		// Add a property row for each metadata struct in the customized Assembly's base schema
+		for (const FAssemblyMetadataDesc& MetadataDesc : BaseSchema->AssemblyMetadata)
 		{
-			ValueWidget = SNew(SCheckBox)
-				.IsChecked_Lambda([this, MetadataDesc]()
-					{
-						bool Value;
-						if (!CustomizedCineAssembly->GetMetadataAsBool(MetadataDesc.Key, Value))
+			if (MetadataDesc.Key.IsEmpty())
+			{
+				continue;
+			}
+
+			CustomizedCineAssembly->AddMetadataNamingToken(MetadataDesc.Key);
+
+			TSharedPtr<SWidget> ValueWidget;
+			if (MetadataDesc.Type == ECineAssemblyMetadataType::String)
+			{
+				ValueWidget = SNew(SBox)
+					.MaxDesiredHeight(120.0f)
+					[
+						MakeStringValueWidget(MetadataDesc)
+					];
+			}
+			else if (MetadataDesc.Type == ECineAssemblyMetadataType::Bool)
+			{
+				ValueWidget = SNew(SCheckBox)
+					.IsChecked_Lambda([this, MetadataDesc]()
 						{
-							Value = MetadataDesc.DefaultValue.Get<bool>();
-						}
-						return Value ? ECheckBoxState::Checked : ECheckBoxState::Unchecked;
-					})
-				.OnCheckStateChanged_Lambda([this, MetadataDesc](ECheckBoxState CheckBoxState)
-					{
-						const bool Value = (CheckBoxState == ECheckBoxState::Checked);
-						CustomizedCineAssembly->SetMetadataAsBool(MetadataDesc.Key, Value);
-					});
-		}
-		else if (MetadataDesc.Type == ECineAssemblyMetadataType::Integer)
-		{
-			ValueWidget = SNew(SNumericEntryBox<int32>)
-				.Font(FAppStyle::GetFontStyle(TEXT("PropertyWindow.NormalFont")))
-				.TypeInterface(LeadingZeroTypeInterface)
-				.Value_Lambda([this, MetadataDesc]()
-					{
-						int32 Value;
-						if (!CustomizedCineAssembly->GetMetadataAsInteger(MetadataDesc.Key, Value))
+							bool Value;
+							if (!CustomizedCineAssembly->GetMetadataAsBool(MetadataDesc.Key, Value))
+							{
+								Value = MetadataDesc.DefaultValue.Get<bool>();
+							}
+							return Value ? ECheckBoxState::Checked : ECheckBoxState::Unchecked;
+						})
+					.OnCheckStateChanged_Lambda([this, MetadataDesc](ECheckBoxState CheckBoxState)
 						{
-							Value = MetadataDesc.DefaultValue.Get<int32>();
-						}
-						return Value;
-					})
-				.OnValueChanged_Lambda([this, MetadataDesc](int32 InValue)
-					{
-						// In order to preserve the desired number of leading zeroes, the integer value is converted to a string using the type interface
-						// and then stored in the assembly as string metadata.
-						const FString ValueString = LeadingZeroTypeInterface->ToString(InValue);
-						CustomizedCineAssembly->SetMetadataAsString(MetadataDesc.Key, ValueString);
-					});
-		}
-		else if (MetadataDesc.Type == ECineAssemblyMetadataType::Float)
-		{
-			ValueWidget = SNew(SNumericEntryBox<float>)
-				.Font(FAppStyle::GetFontStyle(TEXT("PropertyWindow.NormalFont")))
-				.Value_Lambda([this, MetadataDesc]()
-					{
-						float Value;
-						if (!CustomizedCineAssembly->GetMetadataAsFloat(MetadataDesc.Key, Value))
+							const bool Value = (CheckBoxState == ECheckBoxState::Checked);
+							CustomizedCineAssembly->SetMetadataAsBool(MetadataDesc.Key, Value);
+						});
+			}
+			else if (MetadataDesc.Type == ECineAssemblyMetadataType::Integer)
+			{
+				ValueWidget = SNew(SNumericEntryBox<int32>)
+					.Font(FAppStyle::GetFontStyle(TEXT("PropertyWindow.NormalFont")))
+					.TypeInterface(LeadingZeroTypeInterface)
+					.Value_Lambda([this, MetadataDesc]()
 						{
-							Value = MetadataDesc.DefaultValue.Get<float>();
-						}
-						return Value;
-					})
-				.OnValueChanged_Lambda([this, MetadataDesc](float InValue)
-					{
-						CustomizedCineAssembly->SetMetadataAsFloat(MetadataDesc.Key, InValue);
-					});
-		}
-		else if (MetadataDesc.Type == ECineAssemblyMetadataType::AssetPath)
-		{
-			ValueWidget = SNew(SObjectPropertyEntryBox)
-				.AllowedClass(MetadataDesc.AssetClass.ResolveClass())
-				.ThumbnailPool(DetailBuilder.GetThumbnailPool())
-				.AllowCreate(true)
-				.ObjectPath_Lambda([this, MetadataDesc]()
-					{
-						FString Value;
-						if (!CustomizedCineAssembly->GetMetadataAsString(MetadataDesc.Key, Value))
+							int32 Value;
+							if (!CustomizedCineAssembly->GetMetadataAsInteger(MetadataDesc.Key, Value))
+							{
+								Value = MetadataDesc.DefaultValue.Get<int32>();
+							}
+							return Value;
+						})
+					.OnValueChanged_Lambda([this, MetadataDesc](int32 InValue)
 						{
-							Value = MetadataDesc.DefaultValue.Get<FString>();
-						}
-						return Value;
-					})
-				.OnObjectChanged_Lambda([this, MetadataDesc](const FAssetData& InAssetData)
-					{
-						// Store the path of the selected object as a string
-						CustomizedCineAssembly->SetMetadataAsString(MetadataDesc.Key, InAssetData.GetObjectPathString());
-					});
-		}
-		else if (MetadataDesc.Type == ECineAssemblyMetadataType::CineAssembly)
-		{
-			ValueWidget = SNew(SObjectPropertyEntryBox)
-				.AllowedClass(UCineAssembly::StaticClass())
-				.ThumbnailPool(DetailBuilder.GetThumbnailPool())
-				.AllowCreate(true)
-				.OnShouldFilterAsset_Raw(this, &FCineAssemblyCustomization::ShouldFilterAssetBySchema, MetadataDesc.SchemaType)
-				.ObjectPath_Lambda([this, MetadataDesc]()
-					{
-						FString Value;
-						if (!CustomizedCineAssembly->GetMetadataAsString(MetadataDesc.Key, Value))
+							// In order to preserve the desired number of leading zeroes, the integer value is converted to a string using the type interface
+							// and then stored in the assembly as string metadata.
+							const FString ValueString = LeadingZeroTypeInterface->ToString(InValue);
+							CustomizedCineAssembly->SetMetadataAsString(MetadataDesc.Key, ValueString);
+						});
+			}
+			else if (MetadataDesc.Type == ECineAssemblyMetadataType::Float)
+			{
+				ValueWidget = SNew(SNumericEntryBox<float>)
+					.Font(FAppStyle::GetFontStyle(TEXT("PropertyWindow.NormalFont")))
+					.Value_Lambda([this, MetadataDesc]()
 						{
-							Value = MetadataDesc.DefaultValue.Get<FString>();
-						}
-						return Value;
-					})
-				.OnObjectChanged_Lambda([this, MetadataDesc](const FAssetData& InAssetData)
-					{
-						// Store the path of the selected object as a string
-						CustomizedCineAssembly->SetMetadataAsString(MetadataDesc.Key, InAssetData.GetObjectPathString());
-					});
+							float Value;
+							if (!CustomizedCineAssembly->GetMetadataAsFloat(MetadataDesc.Key, Value))
+							{
+								Value = MetadataDesc.DefaultValue.Get<float>();
+							}
+							return Value;
+						})
+					.OnValueChanged_Lambda([this, MetadataDesc](float InValue)
+						{
+							CustomizedCineAssembly->SetMetadataAsFloat(MetadataDesc.Key, InValue);
+						});
+			}
+			else if (MetadataDesc.Type == ECineAssemblyMetadataType::AssetPath)
+			{
+				ValueWidget = SNew(SObjectPropertyEntryBox)
+					.AllowedClass(MetadataDesc.AssetClass.ResolveClass())
+					.ThumbnailPool(DetailBuilder.GetThumbnailPool())
+					.AllowCreate(true)
+					.ObjectPath_Lambda([this, MetadataDesc]()
+						{
+							FString Value;
+							if (!CustomizedCineAssembly->GetMetadataAsString(MetadataDesc.Key, Value))
+							{
+								Value = MetadataDesc.DefaultValue.Get<FString>();
+							}
+							return Value;
+						})
+					.OnObjectChanged_Lambda([this, MetadataDesc](const FAssetData& InAssetData)
+						{
+							// Store the path of the selected object as a string
+							CustomizedCineAssembly->SetMetadataAsString(MetadataDesc.Key, InAssetData.GetObjectPathString());
+						});
+			}
+			else if (MetadataDesc.Type == ECineAssemblyMetadataType::CineAssembly)
+			{
+				ValueWidget = SNew(SObjectPropertyEntryBox)
+					.AllowedClass(UCineAssembly::StaticClass())
+					.ThumbnailPool(DetailBuilder.GetThumbnailPool())
+					.AllowCreate(true)
+					.OnShouldFilterAsset_Raw(this, &FCineAssemblyCustomization::ShouldFilterAssetBySchema, MetadataDesc.SchemaType)
+					.ObjectPath_Lambda([this, MetadataDesc]()
+						{
+							FString Value;
+							if (!CustomizedCineAssembly->GetMetadataAsString(MetadataDesc.Key, Value))
+							{
+								Value = MetadataDesc.DefaultValue.Get<FString>();
+							}
+							return Value;
+						})
+					.OnObjectChanged_Lambda([this, MetadataDesc](const FAssetData& InAssetData)
+						{
+							// Store the path of the selected object as a string
+							CustomizedCineAssembly->SetMetadataAsString(MetadataDesc.Key, InAssetData.GetObjectPathString());
+						});
+			}
+			else
+			{
+				checkNoEntry();
+			}
+
+			MetadataCategory.AddCustomRow(FText::FromString(MetadataDesc.Key))
+				.NameContent()
+				[
+					SNew(STextBlock)
+						.Text(FText::FromString(MetadataDesc.Key))
+						.Font(IDetailLayoutBuilder::GetDetailFont())
+				]
+				.ValueContent()
+				.HAlign(HAlign_Fill)
+				[
+					ValueWidget.ToSharedRef()
+				];
 		}
-		else
+	}
+	else
+	{
+		const FText MetadataCategoryName = LOCTEXT("AllMetadataCategoryName", "Metadata");
+		IDetailCategoryBuilder& MetadataCategory = DetailBuilder.EditCategory("SchemaMetadata", MetadataCategoryName);
+
+		for (const FString& MetadataKey : CustomizedCineAssembly->GetMetadataKeys())
 		{
-			checkNoEntry();
+			if (MetadataKey.IsEmpty())
+			{
+				continue;
+			}
+
+			FString Value;
+			if (CustomizedCineAssembly->GetMetadataAsString(MetadataKey, Value))
+			{
+				MetadataCategory.AddCustomRow(FText::FromString(MetadataKey))
+					.NameContent()
+					[
+						SNew(STextBlock)
+							.Text(FText::FromString(MetadataKey))
+							.Font(IDetailLayoutBuilder::GetDetailFont())
+					]
+					.ValueContent()
+					.HAlign(HAlign_Fill)
+					[
+						SNew(STextBlock)
+							.Text(FText::FromString(Value))
+							.Font(IDetailLayoutBuilder::GetDetailFont())
+					];
+			}
 		}
-
-		MetadataCategory.AddCustomRow(FText::FromString(MetadataDesc.Key))
-			.NameContent()
-			[
-				SNew(STextBlock)
-					.Text(FText::FromString(MetadataDesc.Key))
-					.Font(IDetailLayoutBuilder::GetDetailFont())
-			]
-			.ValueContent()
-			.HAlign(HAlign_Fill)
-			[
-				ValueWidget.ToSharedRef()
-			];
 	}
 }
 
@@ -337,27 +373,45 @@ TSharedRef<SWidget> FCineAssemblyCustomization::MakeStringValueWidget(const FAss
 			.AllowMultiLine(true)
 			.Text_Lambda([this, MetadataDesc]() -> FText
 				{
-					FString Value;
-					const bool bEvaluateTokens = false;
-					if (!CustomizedCineAssembly->GetMetadataAsString(MetadataDesc.Key, Value, bEvaluateTokens))
+					FTemplateString Value;
+					if (!CustomizedCineAssembly->GetMetadataAsTokenString(MetadataDesc.Key, Value))
 					{
-						Value = MetadataDesc.DefaultValue.Get<FString>();
+						Value.Template = MetadataDesc.DefaultValue.Get<FString>();
 					}
-					return FText::FromString(Value);
+					return FText::FromString(Value.Template);
 				})
 			.ResolvedText_Lambda([this, MetadataDesc]()
 				{
 					FString Value;
 					if (!CustomizedCineAssembly->GetMetadataAsString(MetadataDesc.Key, Value))
 					{
-						Value = MetadataDesc.DefaultValue.Get<FString>();
+						return FText::GetEmpty();
 					}
 					return FText::FromString(Value);
 				})
 			.OnTextCommitted_Lambda([this, MetadataDesc](const FText& InText, ETextCommit::Type InCommitType)
 				{
-					CustomizedCineAssembly->SetMetadataAsString(MetadataDesc.Key, InText.ToString());
+					FTemplateString TemplateString;
+					CustomizedCineAssembly->GetMetadataAsTokenString(MetadataDesc.Key, TemplateString);
+					TemplateString.Template = InText.ToString();
+					CustomizedCineAssembly->SetMetadataAsTokenString(MetadataDesc.Key, TemplateString);
 				});
+
+		constexpr float TimerFrequency = 1.0f;
+		StringValueWidget->RegisterActiveTimer(TimerFrequency, FWidgetActiveTimerDelegate::CreateLambda([this, MetadataDesc](double InCurrentTime, float InDeltaTime) -> EActiveTimerReturnType
+			{
+				FTemplateString Value;
+				if (!CustomizedCineAssembly->GetMetadataAsTokenString(MetadataDesc.Key, Value))
+				{
+					Value.Template = MetadataDesc.DefaultValue.Get<FString>();
+				}
+
+				EvaluateTokenString(Value);
+				CustomizedCineAssembly->SetMetadataAsTokenString(MetadataDesc.Key, Value);
+
+				return EActiveTimerReturnType::Continue;
+			}));
+
 	}
 
 	return StringValueWidget.ToSharedRef();
diff --git a/Engine/Plugins/VirtualProduction/CinematicAssemblyTools/Source/CineAssemblyToolsEditor/Private/CineAssemblyFactory.cpp b/Engine/Plugins/VirtualProduction/CinematicAssemblyTools/Source/CineAssemblyToolsEditor/Private/CineAssemblyFactory.cpp
--- a/Engine/Plugins/VirtualProduction/CinematicAssemblyTools/Source/CineAssemblyToolsEditor/Private/CineAssemblyFactory.cpp
+++ b/Engine/Plugins/VirtualProduction/CinematicAssemblyTools/Source/CineAssemblyToolsEditor/Private/CineAssemblyFactory.cpp
@@ -185,6 +185,23 @@ void UCineAssemblyFactory::CreateConfiguredAssembly(UCineAssembly* ConfiguredAss
 		}
 	}
 
+	// Re-evaluate the assembly's metadata tokens
+	if (const UCineAssemblySchema* Schema = ConfiguredAssembly->GetSchema())
+	{
+		for (const FAssemblyMetadataDesc& MetadataDesc : Schema->AssemblyMetadata)
+		{
+			if (MetadataDesc.Type == ECineAssemblyMetadataType::String && MetadataDesc.bEvaluateTokens)
+			{
+				FTemplateString TemplateString;
+				if (ConfiguredAssembly->GetMetadataAsTokenString(MetadataDesc.Key, TemplateString))
+				{
+					TemplateString.Resolved = UCineAssemblyNamingTokens::GetResolvedText(TemplateString.Template, ConfiguredAssembly);
+					ConfiguredAssembly->SetMetadataAsTokenString(MetadataDesc.Key, TemplateString);
+				}
+			}
+		}
+	}
+
 	ConfiguredAssembly->CreateSubAssemblies();
 
 	// Analytics
diff --git a/Engine/Plugins/VirtualProduction/CinematicAssemblyTools/Source/CineAssemblyToolsEditor/Private/CineAssemblyMetadataCustomization.cpp b/Engine/Plugins/VirtualProduction/CinematicAssemblyTools/Source/CineAssemblyToolsEditor/Private/CineAssemblyMetadataCustomization.cpp
--- a/Engine/Plugins/VirtualProduction/CinematicAssemblyTools/Source/CineAssemblyToolsEditor/Private/CineAssemblyMetadataCustomization.cpp
+++ b/Engine/Plugins/VirtualProduction/CinematicAssemblyTools/Source/CineAssemblyToolsEditor/Private/CineAssemblyMetadataCustomization.cpp
@@ -86,12 +86,9 @@ TSharedRef<SWidget> FCineAssemblyMetadataCustomization::MakeStringDefaultValueWi
 					})
 				.ResolvedText_Lambda([&MetadataDesc]()
 					{
-						UNamingTokensEngineSubsystem* NamingTokensSubsystem = GEngine->GetEngineSubsystem<UNamingTokensEngineSubsystem>();
-						FNamingTokenFilterArgs FilterArgs;
-						if (const FString* DefaultValue = MetadataDesc.DefaultValue.TryGet<FString>())
+						if (const FString* Value = MetadataDesc.DefaultValue.TryGet<FString>())
 						{
-							FNamingTokenResultData Result = NamingTokensSubsystem->EvaluateTokenString(*DefaultValue, FilterArgs, {});
-							return Result.EvaluatedText;
+							return FText::FromString(*Value);
 						}
 						return FText::GetEmpty();
 					})
diff --git a/Engine/Plugins/VirtualProduction/CinematicAssemblyTools/Source/CineAssemblyToolsEditor/Private/ProductionFunctionLibrary.cpp b/Engine/Plugins/VirtualProduction/CinematicAssemblyTools/Source/CineAssemblyToolsEditor/Private/ProductionFunctionLibrary.cpp
--- a/Engine/Plugins/VirtualProduction/CinematicAssemblyTools/Source/CineAssemblyToolsEditor/Private/ProductionFunctionLibrary.cpp
+++ b/Engine/Plugins/VirtualProduction/CinematicAssemblyTools/Source/CineAssemblyToolsEditor/Private/ProductionFunctionLibrary.cpp
@@ -5,6 +5,7 @@
 #include "Algo/Find.h"
 #include "CineAssembly.h"
 #include "CineAssemblyFactory.h"
+#include "CineAssemblyNamingTokens.h"
 #include "CineAssemblySchema.h"
 #include "Misc/PackageName.h"
 #include "UObject/Package.h"
@@ -157,13 +158,29 @@ void UProductionFunctionLibrary::AddMetadataToAssembly(UCineAssembly* Assembly,
 
 	for (const TPair<FString, FString>& Pair : Metadata)
 	{
+		if (Pair.Key.IsEmpty())
+		{
+			continue;
+		}
+
 		const FAssemblyMetadataDesc* MetadataDesc = Schema ? Algo::FindBy(Schema->AssemblyMetadata, Pair.Key, &FAssemblyMetadataDesc::Key) : nullptr;
 
 		if (MetadataDesc)
 		{
 			if (MetadataDesc->Type == ECineAssemblyMetadataType::String)
 			{
-				Assembly->SetMetadataAsString(Pair.Key, Pair.Value);
+				if (MetadataDesc->bEvaluateTokens)
+				{
+					FTemplateString TemplateString;
+					TemplateString.Template = Pair.Value;
+					TemplateString.Resolved = UCineAssemblyNamingTokens::GetResolvedText(TemplateString.Template, Assembly);
+
+					Assembly->SetMetadataAsTokenString(Pair.Key, TemplateString);
+				}
+				else
+				{
+					Assembly->SetMetadataAsString(Pair.Key, Pair.Value);
+				}
 			}
 			else if (MetadataDesc->Type == ECineAssemblyMetadataType::Bool)
 			{
diff --git a/Engine/Plugins/VirtualProduction/CinematicAssemblyTools/Source/CineAssemblyToolsEditor/Private/UI/CineAssembly/SCineAssemblyPropertyEditor.cpp b/Engine/Plugins/VirtualProduction/CinematicAssemblyTools/Source/CineAssemblyToolsEditor/Private/UI/CineAssembly/SCineAssemblyPropertyEditor.cpp
--- a/Engine/Plugins/VirtualProduction/CinematicAssemblyTools/Source/CineAssemblyToolsEditor/Private/UI/CineAssembly/SCineAssemblyPropertyEditor.cpp
+++ b/Engine/Plugins/VirtualProduction/CinematicAssemblyTools/Source/CineAssemblyToolsEditor/Private/UI/CineAssembly/SCineAssemblyPropertyEditor.cpp
@@ -7,6 +7,7 @@
 #include "CineAssemblyToolsStyle.h"
 #include "Modules/ModuleManager.h"
 #include "ObjectTools.h"
+#include "PropertyEditorDelegates.h"
 #include "PropertyEditorModule.h"
 #include "ThumbnailRendering/ThumbnailManager.h"
 #include "Widgets/Input/SSegmentedControl.h"
@@ -71,6 +72,7 @@ TSharedRef<SWidget> SCineAssemblyPropertyEditor::BuildUI()
 	DetailsViewArgs.bHideSelectionTip = true;
 
 	TSharedRef<IDetailsView> DetailsView = PropertyEditor.CreateDetailView(DetailsViewArgs);
+	DetailsView->SetIsPropertyVisibleDelegate(FIsPropertyVisible::CreateSP(this, &SCineAssemblyPropertyEditor::IsPropertyVisible));
 	DetailsView->SetIsCustomRowVisibleDelegate(FIsCustomRowVisible::CreateSP(this, &SCineAssemblyPropertyEditor::IsCustomRowVisible));
 
 	DetailsView->SetObject(CineAssembly, true);
@@ -119,6 +121,16 @@ TSharedRef<SWidget> SCineAssemblyPropertyEditor::BuildUI()
 		];
 }
 
+bool SCineAssemblyPropertyEditor::IsPropertyVisible(const FPropertyAndParent& PropertyAndParent)
+{
+	const FName PropertyName = PropertyAndParent.Property.GetFName();
+	if (PropertyName == GET_MEMBER_NAME_CHECKED(UCineAssembly, InstanceMetadata))
+	{
+		return false;
+	}
+	return true;
+}
+
 bool SCineAssemblyPropertyEditor::IsCustomRowVisible(FName RowName, FName ParentName)
 {
 	if (RowName == GET_MEMBER_NAME_CHECKED(UCineAssembly, SubAssemblyNames))
diff --git a/Engine/Plugins/VirtualProduction/CinematicAssemblyTools/Source/CineAssemblyToolsEditor/Private/UI/CineAssembly/SCineAssemblyPropertyEditor.h b/Engine/Plugins/VirtualProduction/CinematicAssemblyTools/Source/CineAssemblyToolsEditor/Private/UI/CineAssembly/SCineAssemblyPropertyEditor.h
--- a/Engine/Plugins/VirtualProduction/CinematicAssemblyTools/Source/CineAssemblyToolsEditor/Private/UI/CineAssembly/SCineAssemblyPropertyEditor.h
+++ b/Engine/Plugins/VirtualProduction/CinematicAssemblyTools/Source/CineAssemblyToolsEditor/Private/UI/CineAssembly/SCineAssemblyPropertyEditor.h
@@ -7,6 +7,8 @@
 #include "CineAssembly.h"
 #include "Widgets/Layout/SWidgetSwitcher.h"
 
+struct FPropertyAndParent;
+
 /** 
  * A panel that displays the editable properties of a Cine Assembly asset 
  */
@@ -45,6 +47,9 @@ private:
 	/** Creates the widget to display for the Overview tab */
 	TSharedRef<SWidget> MakeOverviewWidget();
 
+	/** Filter used by the Details View to determine which properties to show */
+	bool IsPropertyVisible(const FPropertyAndParent& PropertyAndParent);
+
 	/** Filter used by the Details View to determine which custom rows to show */
 	bool IsCustomRowVisible(FName RowName, FName ParentName);
 
