ShotManagement: Fix the logic in ProductionSettings for caching and restoring project default asset names.
This fixes an issue where stale default names can remain registered even after the user has switched to having no Active production.
#rnx

diff --git a/Engine/Plugins/VirtualProduction/CinematicAssemblyTools/Source/CineAssemblyToolsEditor/Private/ProductionSettings.cpp b/Engine/Plugins/VirtualProduction/CinematicAssemblyTools/Source/CineAssemblyToolsEditor/Private/ProductionSettings.cpp
--- a/Engine/Plugins/VirtualProduction/CinematicAssemblyTools/Source/CineAssemblyToolsEditor/Private/ProductionSettings.cpp
+++ b/Engine/Plugins/VirtualProduction/CinematicAssemblyTools/Source/CineAssemblyToolsEditor/Private/ProductionSettings.cpp
@@ -8,6 +8,7 @@
 #include "CineAssemblyToolsAnalytics.h"
 #include "CineAssemblyToolsEditorModule.h"
 #include "Engine/Engine.h"
+#include "Factories/Factory.h"
 #include "HAL/FileManager.h"
 #include "HAL/PlatformFileManager.h"
 #include "LevelSequenceProjectSettings.h"
@@ -253,6 +254,19 @@ void UProductionSettings::CacheProjectDefaults()
 
 	const UMovieSceneToolsProjectSettings* const MovieSceneToolsSettings = GetDefault<UMovieSceneToolsProjectSettings>();
 	ProjectDefaultStartTime = MovieSceneToolsSettings->DefaultStartTime;
+
+	// Cache any registered default asset names for asset classes that could be overridden by the production settings
+	IAssetTools& AssetTools = FModuleManager::GetModuleChecked<FAssetToolsModule>("AssetTools").Get();
+	TArray<UFactory*> Factories = AssetTools.GetNewAssetFactories();
+	for (UFactory* Factory : Factories)
+	{
+		const UClass* Class = Factory->GetSupportedClass();
+		const TOptional<FString> DefaultName = AssetTools.GetDefaultAssetNameForClass(Class);
+		if (DefaultName.IsSet())
+		{
+			ProjectDefaultAssetNames.Add(Class, DefaultName.GetValue());
+		}
+	}
 }
 
 FStructView UProductionSettings::GetMutableProductionExtendedData(FGuid ProductionID, const UScriptStruct& DataStruct)
@@ -421,22 +435,28 @@ void UProductionSettings::OverrideSubsequenceHierarchicalBias()
 
 void UProductionSettings::OverrideDefaultAssetNames()
 {
+	IAssetTools& AssetTools = FModuleManager::GetModuleChecked<FAssetToolsModule>("AssetTools").Get();
+
+	// Update the default asset names for each class that has one specified in the active production
+	// If there is no active production, all project default asset names will be re-registered
 	if (const FCinematicProduction* ActiveProduction = GetActiveProductionPtr())
 	{
-		IAssetTools& AssetTools = FModuleManager::GetModuleChecked<FAssetToolsModule>("AssetTools").Get();
-
 		// Register all of the active production's default asset names with the asset tools
 		for (const TPair<TObjectPtr<const UClass>, FString>& Naming : ActiveProduction->DefaultAssetNames)
 		{
 			if (Naming.Key)
 			{
-				// Before registering a new default name with AssetTools, cache the existing value (if it exists) so we can reset it later
-				TOptional<FString> ExistingDefaultName = AssetTools.GetDefaultAssetNameForClass(Naming.Key);
-				if (ExistingDefaultName.IsSet() && !ProjectDefaultAssetNames.Contains(Naming.Key))
-				{
-					ProjectDefaultAssetNames.Add(Naming.Key, ExistingDefaultName.GetValue());
-				}
-
+				AssetTools.RegisterDefaultAssetNameForClass(Naming.Key, Naming.Value);
+			}
+		}
+	}
+	else
+	{
+		// Register all of the active production's default asset names with the asset tools
+		for (const TPair<const UClass*, FString>& Naming : ProjectDefaultAssetNames)
+		{
+			if (Naming.Key)
+			{
 				AssetTools.RegisterDefaultAssetNameForClass(Naming.Key, Naming.Value);
 			}
 		}
@@ -767,11 +787,6 @@ void UProductionSettings::SetActiveProduction(FGuid ProductionID)
 		for (const TPair<TObjectPtr<const UClass>, FString>& Naming : ActiveProduction->DefaultAssetNames)
 		{
 			AssetTools.UnregisterDefaultAssetNameForClass(Naming.Key);
-
-			if (ProjectDefaultAssetNames.Contains(Naming.Key))
-			{
-				AssetTools.RegisterDefaultAssetNameForClass(Naming.Key, ProjectDefaultAssetNames[Naming.Key]);
-			}
 		}
 	}
 
@@ -941,15 +956,6 @@ void UProductionSettings::AddAssetNaming(FGuid ProductionID, const UClass* Asset
 		if (IsActiveProduction(ProductionID))
 		{
 			IAssetTools& AssetTools = FModuleManager::GetModuleChecked<FAssetToolsModule>("AssetTools").Get();
-
-			// If there is an existing name registered with AssetTools for this class, cache it before overriding it
-			TOptional<FString> ExistingDefaultName = AssetTools.GetDefaultAssetNameForClass(AssetClass);
-			if (ExistingDefaultName.IsSet() && !ProjectDefaultAssetNames.Contains(AssetClass))
-			{
-				ProjectDefaultAssetNames.Add(AssetClass, ExistingDefaultName.GetValue());
-			}
-
-			// Register the new default asset name with AssetTools
 			AssetTools.RegisterDefaultAssetNameForClass(AssetClass, DefaultName);
 		}
 
