ShotManagement: Fix potential crash on startup caused by [REDACTED_PERSON] trying to access the Asset Tools module when it wasn't loaded yet.

#rnx

diff --git a/Engine/Plugins/VirtualProduction/CinematicAssemblyTools/Source/CineAssemblyToolsEditor/Private/ProductionSettings.cpp b/Engine/Plugins/VirtualProduction/CinematicAssemblyTools/Source/CineAssemblyToolsEditor/Private/ProductionSettings.cpp
--- a/Engine/Plugins/VirtualProduction/CinematicAssemblyTools/Source/CineAssemblyToolsEditor/Private/ProductionSettings.cpp
+++ b/Engine/Plugins/VirtualProduction/CinematicAssemblyTools/Source/CineAssemblyToolsEditor/Private/ProductionSettings.cpp
@@ -256,7 +256,7 @@ void UProductionSettings::CacheProjectDefaults()
 	ProjectDefaultStartTime = MovieSceneToolsSettings->DefaultStartTime;
 
 	// Cache any registered default asset names for asset classes that could be overridden by the production settings
-	IAssetTools& AssetTools = FModuleManager::GetModuleChecked<FAssetToolsModule>("AssetTools").Get();
+	IAssetTools& AssetTools = FModuleManager::LoadModuleChecked<FAssetToolsModule>("AssetTools").Get();
 	TArray<UFactory*> Factories = AssetTools.GetNewAssetFactories();
 	for (UFactory* Factory : Factories)
 	{
