ShotManagement: Fix editor crash in assembly duplication tool when duplicating a read-only assembly.

diff --git a/Engine/Plugins/VirtualProduction/CinematicAssemblyTools/Source/CineAssemblyToolsEditor/Private/UI/CineAssembly/SDuplicateAssemblyWindow.cpp b/Engine/Plugins/VirtualProduction/CinematicAssemblyTools/Source/CineAssemblyToolsEditor/Private/UI/CineAssembly/SDuplicateAssemblyWindow.cpp
--- a/Engine/Plugins/VirtualProduction/CinematicAssemblyTools/Source/CineAssemblyToolsEditor/Private/UI/CineAssembly/SDuplicateAssemblyWindow.cpp
+++ b/Engine/Plugins/VirtualProduction/CinematicAssemblyTools/Source/CineAssemblyToolsEditor/Private/UI/CineAssembly/SDuplicateAssemblyWindow.cpp
@@ -189,19 +189,12 @@ TSharedRef<SWidget> SDuplicateAssemblyWindow::MakeSubsequencesPanel()
 				]
 
 			+ SVerticalBox::Slot()
-				.AutoHeight()
+				.FillHeight(1.0f)
 				.Padding(0.0f, 0.0f, 0.0f, 16.0f)
 				[
 					TreeView.ToSharedRef()
 				]
 
-			// The spacer fills the remaining space so the path picker will appear at the bottom of the panel
-			+ SVerticalBox::Slot()
-				.FillHeight(1.0f)
-				[
-					SNew(SSpacer)
-				]
-
 			+ SVerticalBox::Slot()
 				.AutoHeight()
 				[
@@ -396,6 +389,11 @@ void SDuplicateAssemblyWindow::DuplicateSubsequencesRecursive(UMovieSceneSequenc
 		return;
 	}
 
+	// In order to duplicate subsequences, the MovieScene needs to be mutable.
+	// The original read-only state will be restored when we are done.
+	const bool bIsReadOnly = MovieScene->IsReadOnly();
+	MovieScene->SetReadOnly(false);
+
 	IAssetTools& AssetTools = FModuleManager::GetModuleChecked<FAssetToolsModule>("AssetTools").Get();
 
 	// Loop through all of the tracks and all of the sections to find every subsequence in the input sequence
@@ -409,6 +407,11 @@ void SDuplicateAssemblyWindow::DuplicateSubsequencesRecursive(UMovieSceneSequenc
 				{
 					if (UMovieSceneSequence* Subsequence = SubSection->GetSequence())
 					{
+						if (!SubsequenceDuplicationData.Contains(Subsequence))
+						{
+							continue;
+						}
+
 						const FSubsequenceDuplicationData& DuplicationData = SubsequenceDuplicationData[Subsequence];
 
 						if (DuplicationData.DuplicationMode == ECineAssemblyDuplicationMode::DuplicateOriginal)
@@ -437,12 +440,19 @@ void SDuplicateAssemblyWindow::DuplicateSubsequencesRecursive(UMovieSceneSequenc
 							const FString SubAssemblyPath = FPaths::GetPath(UniquePackageName);
 							UMovieSceneSequence* DuplicateSubsequence = Cast<UMovieSceneSequence>(AssetTools.DuplicateAsset(UniqueAssetName, SubAssemblyPath, Subsequence));
 
-							// Update the sequence of the duplicate subsection to be the duplicate subsequence
-							SubSection->SetSequence(DuplicateSubsequence);
-
-							// Update the data in the duplication data map to correspond to the new duplicate subsequence
-							SubsequenceDuplicationData.Add(DuplicateSubsequence, DuplicationData);
-							SubsequenceDuplicationData.Remove(Subsequence);
+							if (DuplicateSubsequence)
+							{
+								// Update the sequence of the duplicate subsection to be the duplicate subsequence
+								// If the section is locked, unlock it, then update the subsequence, then lock it again
+								const bool bIsLocked = SubSection->IsLocked();
+								SubSection->SetIsLocked(false);
+								SubSection->SetSequence(DuplicateSubsequence);
+								SubSection->SetIsLocked(bIsLocked);
+
+								// Update the data in the duplication data map to correspond to the new duplicate subsequence
+								SubsequenceDuplicationData.Add(DuplicateSubsequence, DuplicationData);
+								SubsequenceDuplicationData.Remove(Subsequence);
+							}
 						}
 
 						DuplicateSubsequencesRecursive(SubSection->GetSequence());
@@ -451,6 +461,9 @@ void SDuplicateAssemblyWindow::DuplicateSubsequencesRecursive(UMovieSceneSequenc
 			}
 		}
 	}
+
+	// Restore the original read-only state
+	MovieScene->SetReadOnly(bIsReadOnly);
 }
 
 void SDuplicateAssemblyWindow::RemoveSubsequencesRecursive(UMovieSceneSequence* Sequence)
@@ -478,6 +491,11 @@ void SDuplicateAssemblyWindow::RemoveSubsequencesRecursive(UMovieSceneSequence*
 				{
 					if (UMovieSceneSequence* Subsequence = SubSection->GetSequence())
 					{
+						if (!SubsequenceDuplicationData.Contains(Subsequence))
+						{
+							continue;
+						}
+
 						const FSubsequenceDuplicationData& DuplicationData = SubsequenceDuplicationData[Subsequence];
 						if (DuplicationData.DuplicationMode == ECineAssemblyDuplicationMode::Remove)
 						{
