ShotMangement: Replace inline editable textbox widgets with the new naming token textbox widget in the Production Wizard for asset naming and template folders.
#rnx

diff --git a/Engine/Plugins/VirtualProduction/CinematicAssemblyTools/Source/CineAssemblyToolsEditor/CineAssemblyToolsEditor.build.cs b/Engine/Plugins/VirtualProduction/CinematicAssemblyTools/Source/CineAssemblyToolsEditor/CineAssemblyToolsEditor.build.cs
--- a/Engine/Plugins/VirtualProduction/CinematicAssemblyTools/Source/CineAssemblyToolsEditor/CineAssemblyToolsEditor.build.cs
+++ b/Engine/Plugins/VirtualProduction/CinematicAssemblyTools/Source/CineAssemblyToolsEditor/CineAssemblyToolsEditor.build.cs
@@ -38,6 +38,7 @@ namespace UnrealBuildTool.Rules
 					"MovieScene",
 					"MovieSceneTools",
 					"NamingTokens",
+					"NamingTokensUI",
 					"Projects",
 					"PropertyEditor",
 					"SharedSettingsWidgets",
diff --git a/Engine/Plugins/VirtualProduction/CinematicAssemblyTools/Source/CineAssemblyToolsEditor/Private/UI/SAssetNamingPanel.cpp b/Engine/Plugins/VirtualProduction/CinematicAssemblyTools/Source/CineAssemblyToolsEditor/Private/UI/SAssetNamingPanel.cpp
--- a/Engine/Plugins/VirtualProduction/CinematicAssemblyTools/Source/CineAssemblyToolsEditor/Private/UI/SAssetNamingPanel.cpp
+++ b/Engine/Plugins/VirtualProduction/CinematicAssemblyTools/Source/CineAssemblyToolsEditor/Private/UI/SAssetNamingPanel.cpp
@@ -16,7 +16,6 @@
 #include "UI/SActiveProductionCombo.h"
 #include "Widgets/Layout/SSeparator.h"
 #include "Widgets/SBoxPanel.h"
-#include "Widgets/Text/SInlineEditableTextBlock.h"
 
 #define LOCTEXT_NAMESPACE "SAssetNamingPanel"
 
@@ -208,6 +207,7 @@ void SAssetNamingPanel::Construct(const FArguments& InArgs)
 									SAssignNew(AssetNamingListView, SListView<TSharedPtr<FAssetNamingRowData>>)
 										.ListItemsSource(&AssetNamingListItems)
 										.OnGenerateRow(this, &SAssetNamingPanel::OnGenerateAssetNamingRow)
+										.OnMouseButtonDoubleClick(this, &SAssetNamingPanel::OnListViewDoubleClick)
 										.SelectionMode(ESelectionMode::None)
 										.HeaderRow
 										(
@@ -268,6 +268,14 @@ TSharedRef<ITableRow> SAssetNamingPanel::OnGenerateAssetNamingRow(TSharedPtr<FAs
 		.OnDeleteRow(this, &SAssetNamingPanel::UpdateAssetNamingList);
 }
 
+void SAssetNamingPanel::OnListViewDoubleClick(TSharedPtr<FAssetNamingRowData> InAssetNaming)
+{
+	if (TSharedPtr<SAssetNamingRow> Widget = StaticCastSharedPtr<SAssetNamingRow>(AssetNamingListView->WidgetFromItem(InAssetNaming)))
+	{
+		Widget->EditableTextBlock->EnterEditingMode();
+	}
+}
+
 void SAssetNamingRow::Construct(const FArguments& Args, const TSharedRef<STableViewBase>& OwnerTableView, const TSharedPtr<FAssetNamingRowData>& InAssetNaming)
 {
 	AssetNaming = InAssetNaming;
@@ -303,6 +311,17 @@ TSharedRef<SWidget> SAssetNamingRow::GenerateWidgetForColumn(const FName& Column
 	}
 	else if (ColumnName == NamingColumnName)
 	{
+		EditableTextBlock = SNew(SNamingTokensEditableTextBox)
+			.SupportInlineEdit(true)
+			.ShouldEvaluateTokens(false)
+			.DisplayTokenIcon(false)
+			.DisplayBorderImage(false)
+			.IsReadOnly_Lambda([this]() { return (AssetNaming->Class == nullptr); })
+			.DisplayErrorMessage_Lambda([this]() { return (AssetNaming->Class != nullptr); })
+			.Text_Lambda([this]() { return FText::FromString(AssetNaming->DefaultName); })
+			.OnTextCommitted(this, &SAssetNamingRow::SetDefaultAssetName)
+			.OnValidateTokenizedText(this, &SAssetNamingRow::ValidateDefaultAssetName);
+
 		return SNew(SHorizontalBox)
 
 			+ SHorizontalBox::Slot()
@@ -313,11 +332,7 @@ TSharedRef<SWidget> SAssetNamingRow::GenerateWidgetForColumn(const FName& Column
 					.Padding(FMargin(4.0, 1.0f))
 					.OnMouseButtonDown(this, &SAssetNamingRow::SummonEditMenu)
 					[
-						SAssignNew(EditableTextBlock, SInlineEditableTextBlock)
-							.Text_Lambda([this]() { return FText::FromString(AssetNaming->DefaultName); })
-							.IsEnabled_Lambda([this]() { return (AssetNaming->Class != nullptr); })
-							.OnVerifyTextChanged(this, &SAssetNamingRow::ValidateDefaultAssetName)
-							.OnTextCommitted(this, &SAssetNamingRow::SetDefaultAssetName)
+						EditableTextBlock.ToSharedRef()
 					]
 			];
 	}
@@ -424,7 +439,6 @@ FReply SAssetNamingRow::SummonEditMenu(const FGeometry& MyGeometry, const FPoint
 			FSlateIcon(FCineAssemblyToolsStyle::StyleName, "Icons.AssetNaming"),
 			FUIAction(FExecuteAction::CreateLambda([this]()
 				{
-					FSlateApplication::Get().SetKeyboardFocus(EditableTextBlock.ToSharedRef());
 					EditableTextBlock->EnterEditingMode();
 				}
 			)),
diff --git a/Engine/Plugins/VirtualProduction/CinematicAssemblyTools/Source/CineAssemblyToolsEditor/Private/UI/SAssetNamingPanel.h b/Engine/Plugins/VirtualProduction/CinematicAssemblyTools/Source/CineAssemblyToolsEditor/Private/UI/SAssetNamingPanel.h
--- a/Engine/Plugins/VirtualProduction/CinematicAssemblyTools/Source/CineAssemblyToolsEditor/Private/UI/SAssetNamingPanel.h
+++ b/Engine/Plugins/VirtualProduction/CinematicAssemblyTools/Source/CineAssemblyToolsEditor/Private/UI/SAssetNamingPanel.h
@@ -4,8 +4,8 @@
 
 #include "Widgets/SCompoundWidget.h"
 
+#include "SNamingTokensEditableTextBox.h"
 #include "Widgets/DeclarativeSyntaxSupport.h"
-#include "Widgets/Text/SInlineEditableTextBlock.h"
 #include "Widgets/Views/SListView.h"
 
 /** Data to display in a row in the asset naming list view */
@@ -39,6 +39,10 @@ public:
 	/** Creates the widget for this row for the specified column */
 	virtual TSharedRef<SWidget> GenerateWidgetForColumn(const FName& ColumnName) override;
 
+public:
+	/** The text box where the user can type the default asset name */
+	TSharedPtr<SNamingTokensEditableTextBox> EditableTextBlock;
+
 private:
 	/** Set the asset class for this row widget */
 	void SetAssetClass(const UClass* SelectedClass);
@@ -59,9 +63,6 @@ private:
 	/** The underlying asset naming data, used to properly display the class picker and default naming for that class */
 	TSharedPtr<FAssetNamingRowData> AssetNaming;
 
-	/** The text box where the user can type the default asset name */
-	TSharedPtr<SInlineEditableTextBlock> EditableTextBlock;
-
 	/** Delegate to execute when a row is deleted from the list */
 	FOnDeleteRow OnDeleteRow;
 };
@@ -84,6 +85,9 @@ private:
 	/** Generates a row displaying asset naming data */
 	TSharedRef<ITableRow> OnGenerateAssetNamingRow(TSharedPtr<FAssetNamingRowData> InAssetNaming, const TSharedRef<STableViewBase>& OwnerTable);
 
+	/** Handle double clicks on rows in the list view */
+	void OnListViewDoubleClick(TSharedPtr<FAssetNamingRowData> InAssetNaming);
+
 	/** Update the list view with the asset naming properties of the current active production */
 	void UpdateAssetNamingList();
 
diff --git a/Engine/Plugins/VirtualProduction/CinematicAssemblyTools/Source/CineAssemblyToolsEditor/Private/UI/SFolderHierarchyPanel.cpp b/Engine/Plugins/VirtualProduction/CinematicAssemblyTools/Source/CineAssemblyToolsEditor/Private/UI/SFolderHierarchyPanel.cpp
--- a/Engine/Plugins/VirtualProduction/CinematicAssemblyTools/Source/CineAssemblyToolsEditor/Private/UI/SFolderHierarchyPanel.cpp
+++ b/Engine/Plugins/VirtualProduction/CinematicAssemblyTools/Source/CineAssemblyToolsEditor/Private/UI/SFolderHierarchyPanel.cpp
@@ -29,7 +29,6 @@
 #include "Widgets/Layout/SBorder.h"
 #include "Widgets/Layout/SSeparator.h"
 #include "Widgets/SBoxPanel.h"
-#include "Widgets/Text/SInlineEditableTextBlock.h"
 
 #define LOCTEXT_NAMESPACE "SFolderHierarchyPanel"
 
@@ -271,7 +270,7 @@ TSharedRef<ITableRow> SFolderHierarchyPanel::OnGenerateRow(TSharedPtr<FTemplateF
 
 	return SNew(STableRow<TSharedPtr<FTemplateFolderTreeItem>>, OwnerTable)
 		.ShowSelection(true)
-		.Padding(FMargin(8.0f, 2.0f, 8.0f, 0.0f))
+		.Padding(FMargin(8.0f, 0.0f, 8.0f, 0.0f))
  		.OnCanAcceptDrop_Lambda([this](const FDragDropEvent& InDragDropEvent, EItemDropZone InDropZone, TSharedPtr<FTemplateFolderTreeItem> InItem)
  			{
  				return (InItem != PluginRootItem) ? EItemDropZone::OntoItem : TOptional<EItemDropZone>();
@@ -295,14 +294,20 @@ TSharedRef<ITableRow> SFolderHierarchyPanel::OnGenerateRow(TSharedPtr<FTemplateF
 			+ SHorizontalBox::Slot()
 				.VAlign(VAlign_Center)
 				[
-					SAssignNew(TreeItem->NameWidget, SInlineEditableTextBlock)
+					SAssignNew(TreeItem->NameWidget, SNamingTokensEditableTextBox)
+						.SupportInlineEdit(true)
+						.ShouldEvaluateTokens(false)
+						.TextBoxPadding(FMargin(4, 2, 0, 2))
+						.DisplayTokenIcon(false)
+						.DisplayBorderImage(false)
+						.IsReadOnly(TreeItem->bIsReadOnly)
+						.DisplayErrorMessage(!TreeItem->bIsReadOnly)
 						.Text_Lambda([TreeItem]() -> FText
 							{
 								const FString ItemName = !TreeItem->FriendlyName.IsEmpty() ? TreeItem->FriendlyName : FPaths::GetPathLeaf(TreeItem->Path.GetInternalPathString());
 								return FText::FromString(ItemName);
 							})
-						.IsReadOnly(TreeItem->bIsReadOnly)
-						.OnVerifyTextChanged(this, &SFolderHierarchyPanel::IsValidFolderName, TreeItem)
+						.OnValidateTokenizedText(this, &SFolderHierarchyPanel::IsValidFolderName, TreeItem)
 						.OnTextCommitted(this, &SFolderHierarchyPanel::SetTemplateFolderName, TreeItem)
 				]
 
@@ -328,7 +333,6 @@ void SFolderHierarchyPanel::OnItemsRebuilt()
 	// Upon regenerating the tree view, allow the user to immediately interact with the name widget of the newly added template folder in order to rename it
 	if (MostRecentlyAddedItem && MostRecentlyAddedItem->NameWidget)
 	{
-		FSlateApplication::Get().SetKeyboardFocus(MostRecentlyAddedItem->NameWidget.ToSharedRef());
 		MostRecentlyAddedItem->NameWidget->EnterEditingMode();
 		MostRecentlyAddedItem.Reset();
 	}
@@ -402,7 +406,6 @@ TSharedPtr<SWidget> SFolderHierarchyPanel::OnContextMenuOpening()
 				FUIAction(
 					FExecuteAction::CreateLambda([SelectedTreeItem]()
 						{
-							FSlateApplication::Get().SetKeyboardFocus(SelectedTreeItem->NameWidget.ToSharedRef());
 							SelectedTreeItem->NameWidget->EnterEditingMode();
 						})),
 				NAME_None,
@@ -448,7 +451,6 @@ void SFolderHierarchyPanel::OnTreeViewDoubleClick(TSharedPtr<FTemplateFolderTree
 	// Put the textblock for the input tree item into edit mode so the user can rename the item
 	if (TreeItem && TreeItem->NameWidget)
 	{
-		FSlateApplication::Get().SetKeyboardFocus(TreeItem->NameWidget.ToSharedRef());
 		TreeItem->NameWidget->EnterEditingMode();
 	}
 }
@@ -719,7 +721,12 @@ bool SFolderHierarchyPanel::IsValidFolderName(const FText& InText, FText& OutErr
 
 void SFolderHierarchyPanel::SetTemplateFolderName(const FText& InText, ETextCommit::Type InCommitType, TSharedPtr<FTemplateFolderTreeItem> TreeItem)
 {
-	// Early-out if the folder name has not actually changed
+	// Early-out if the folder is read-only, or if the name has not actually changed
+	if (TreeItem->bIsReadOnly)
+	{
+		return;
+	}
+
 	const FString OldPath = TreeItem->Path.GetInternalPathString();
 	const FString OldFolderName = FPaths::GetPathLeaf(OldPath);
 	if (OldFolderName.Equals(InText.ToString()))
diff --git a/Engine/Plugins/VirtualProduction/CinematicAssemblyTools/Source/CineAssemblyToolsEditor/Private/UI/SFolderHierarchyPanel.h b/Engine/Plugins/VirtualProduction/CinematicAssemblyTools/Source/CineAssemblyToolsEditor/Private/UI/SFolderHierarchyPanel.h
--- a/Engine/Plugins/VirtualProduction/CinematicAssemblyTools/Source/CineAssemblyToolsEditor/Private/UI/SFolderHierarchyPanel.h
+++ b/Engine/Plugins/VirtualProduction/CinematicAssemblyTools/Source/CineAssemblyToolsEditor/Private/UI/SFolderHierarchyPanel.h
@@ -6,7 +6,7 @@
 
 #include "ContentBrowserItemPath.h"
 #include "ProductionSettings.h"
-#include "Widgets/Text/SInlineEditableTextBlock.h"
+#include "SNamingTokensEditableTextBox.h"
 #include "Widgets/Views/STreeView.h"
 
 /** States in which a template folder can exist */
@@ -44,7 +44,7 @@ struct FTemplateFolderTreeItem
 	ETemplateFolderStatus Status = ETemplateFolderStatus::None;
 
 	/** The text widget that display this item's name and supports renaming */
-	TSharedPtr<SInlineEditableTextBlock> NameWidget;
+	TSharedPtr<SNamingTokensEditableTextBox> NameWidget;
 
 	/** If true, this item cannot be renamed or deleted */
 	bool bIsReadOnly = false;
