ShotMangement: Replace old template string textbox widgets with the new naming token textbox widget in Cine Assembly and Schema tools.
This allows us to remove a lot of manual token evaluation in these tools as well, since that can now be handled by the widget directly.
#rnx

diff --git a/Engine/Plugins/VirtualProduction/CinematicAssemblyTools/Source/CineAssemblyToolsEditor/Private/CineAssemblyCustomization.cpp b/Engine/Plugins/VirtualProduction/CinematicAssemblyTools/Source/CineAssemblyToolsEditor/Private/CineAssemblyCustomization.cpp
--- a/Engine/Plugins/VirtualProduction/CinematicAssemblyTools/Source/CineAssemblyToolsEditor/Private/CineAssemblyCustomization.cpp
+++ b/Engine/Plugins/VirtualProduction/CinematicAssemblyTools/Source/CineAssemblyToolsEditor/Private/CineAssemblyCustomization.cpp
@@ -13,6 +13,7 @@
 #include "MovieSceneMetaData.h"
 #include "ProductionSettings.h"
 #include "PropertyCustomizationHelpers.h"
+#include "SNamingTokensEditableTextBox.h"
 #include "STemplateStringEditableTextBox.h"
 #include "Widgets/Input/SCheckBox.h"
 #include "Widgets/Input/SComboButton.h"
@@ -43,6 +44,11 @@ void FCineAssemblyCustomization::CustomizeDetails(IDetailLayoutBuilder& DetailBu
 
 	LeadingZeroTypeInterface = MakeShared<FLeadingZeroNumericTypeInterface>();
 
+	NamingTokenContext = TStrongObjectPtr<UCineAssemblyNamingTokensContext>(NewObject<UCineAssemblyNamingTokensContext>());
+	NamingTokenContext->Assembly = CustomizedCineAssembly;
+
+	FilterArgs.AdditionalNamespacesToInclude.Add(UCineAssemblyNamingTokens::TokenNamespace);
+
 	CustomizeDefaultCategory(DetailBuilder);
 	CustomizeMetadataCategory(DetailBuilder);
 	CustomizeSubsequenceCategory(DetailBuilder);
@@ -337,10 +343,22 @@ void FCineAssemblyCustomization::CustomizeSubsequenceCategory(IDetailLayoutBuild
 			.ValueContent()
 			.HAlign(HAlign_Fill)
 			[
-				SNew(STemplateStringEditableTextBox)
-					.Text_Raw(this, &FCineAssemblyCustomization::GetTemplateText, Index)
-					.ResolvedText_Lambda([this, Index]() { return GetResolvedText(Index); })
-					.OnTextCommitted_Raw(this, &FCineAssemblyCustomization::OnTemplateTextCommitted, Index)
+				SNew(SNamingTokensEditableTextBox)
+					.Contexts({ NamingTokenContext.Get() })
+					.FilterArgs(FilterArgs)
+					.EvaluationFrequency(1.0f)
+					.Text_Lambda([this, Index]()
+						{
+							return GetSubAssemblyTemplateName(Index);
+						})
+					.OnTextCommitted_Lambda([this, Index](const FText& InText, ETextCommit::Type InCommitType)
+						{
+							SetSubAssemblyTemplateName(Index, InText, ETemplateStringType::Template);
+						})
+					.OnTokenizedTextEvaluated_Lambda([this, Index](const FText& InText)
+						{
+							SetSubAssemblyTemplateName(Index, InText, ETemplateStringType::Resolved);
+						})
 			];
 	}
 }
@@ -369,49 +387,23 @@ TSharedRef<SWidget> FCineAssemblyCustomization::MakeStringValueWidget(const FAss
 	}
 	else
 	{
-		StringValueWidget = SNew(STemplateStringEditableTextBox)
+		StringValueWidget = SNew(SNamingTokensEditableTextBox)
 			.AllowMultiLine(true)
-			.Text_Lambda([this, MetadataDesc]() -> FText
+			.Contexts({ NamingTokenContext.Get() })
+			.FilterArgs(FilterArgs)
+			.EvaluationFrequency(1.0f)
+			.Text_Lambda([this, &MetadataDesc]() -> FText
 				{
-					FTemplateString Value;
-					if (!CustomizedCineAssembly->GetMetadataAsTokenString(MetadataDesc.Key, Value))
-					{
-						Value.Template = MetadataDesc.DefaultValue.Get<FString>();
-					}
-					return FText::FromString(Value.Template);
+					return GetMetadataTokenStringValue(MetadataDesc.Key);
 				})
-			.ResolvedText_Lambda([this, MetadataDesc]()
+			.OnTextCommitted_Lambda([this, &MetadataDesc](const FText& InText, ETextCommit::Type InCommitType)
 				{
-					FString Value;
-					if (!CustomizedCineAssembly->GetMetadataAsString(MetadataDesc.Key, Value))
-					{
-						return FText::GetEmpty();
-					}
-					return FText::FromString(Value);
+					SetMetadataTokenStringValue(MetadataDesc.Key, InText, ETemplateStringType::Template);
 				})
-			.OnTextCommitted_Lambda([this, MetadataDesc](const FText& InText, ETextCommit::Type InCommitType)
+			.OnTokenizedTextEvaluated_Lambda([this, &MetadataDesc](const FText& InText)
 				{
-					FTemplateString TemplateString;
-					CustomizedCineAssembly->GetMetadataAsTokenString(MetadataDesc.Key, TemplateString);
-					TemplateString.Template = InText.ToString();
-					CustomizedCineAssembly->SetMetadataAsTokenString(MetadataDesc.Key, TemplateString);
+					SetMetadataTokenStringValue(MetadataDesc.Key, InText, ETemplateStringType::Resolved);
 				});
-
-		constexpr float TimerFrequency = 1.0f;
-		StringValueWidget->RegisterActiveTimer(TimerFrequency, FWidgetActiveTimerDelegate::CreateLambda([this, MetadataDesc](double InCurrentTime, float InDeltaTime) -> EActiveTimerReturnType
-			{
-				FTemplateString Value;
-				if (!CustomizedCineAssembly->GetMetadataAsTokenString(MetadataDesc.Key, Value))
-				{
-					Value.Template = MetadataDesc.DefaultValue.Get<FString>();
-				}
-
-				EvaluateTokenString(Value);
-				CustomizedCineAssembly->SetMetadataAsTokenString(MetadataDesc.Key, Value);
-
-				return EActiveTimerReturnType::Continue;
-			}));
-
 	}
 
 	return StringValueWidget.ToSharedRef();
@@ -460,67 +452,6 @@ void FCineAssemblyCustomization::SubAssemblyCheckStateChanged(ECheckBoxState Che
 	}
 }
 
-FText FCineAssemblyCustomization::GetTemplateText(int32 Index) const
-{
-	if (SubAssemblyNames.IsValidIndex(Index))
-	{
-		const FTemplateString& SubAssemblyName = SubAssemblyNames[Index];
-		return FText::FromString(FPaths::GetBaseFilename(SubAssemblyName.Template));
-	}
-	return FText::GetEmpty();
-}
-
-FText FCineAssemblyCustomization::GetResolvedText(int32 Index)
-{
-	EvaluateTokenStrings();
-
-	if (SubAssemblyNames.IsValidIndex(Index))
-	{
-		const FTemplateString& SubAssemblyName = SubAssemblyNames[Index];
-		return FText::FromString(FPaths::GetBaseFilename(SubAssemblyName.Resolved.ToString()));
-	}
-	return FText::GetEmpty();
-}
-
-void FCineAssemblyCustomization::OnTemplateTextCommitted(const FText& InText, ETextCommit::Type InCommitType, int32 Index)
-{
-	if (SubAssemblyNames.IsValidIndex(Index))
-	{
-		FTemplateString& SubAssemblyName = SubAssemblyNames[Index];
-
-		FTemplateString* TemplateString = Algo::FindBy(CustomizedCineAssembly->SubAssemblyNames, SubAssemblyName.Template, &FTemplateString::Template);
-
-		const FString Path = FPaths::GetPath(SubAssemblyName.Template);
-		SubAssemblyName.Template = Path / InText.ToString();
-
-		if (TemplateString)
-		{
-			TemplateString->Template = SubAssemblyName.Template;
-		}
-
-		EvaluateTokenString(SubAssemblyName);
-	}
-}
-
-void FCineAssemblyCustomization::EvaluateTokenStrings()
-{
-	FDateTime CurrentTime = FDateTime::Now();
-	if ((CurrentTime - LastTokenUpdateTime).GetSeconds() >= 1.0f)
-	{
-		for (FTemplateString& SubAssemblyName : SubAssemblyNames)
-		{
-			EvaluateTokenString(SubAssemblyName);
-		}
-
-		LastTokenUpdateTime = CurrentTime;
-	}
-}
-
-void FCineAssemblyCustomization::EvaluateTokenString(FTemplateString& TokenString)
-{
-	TokenString.Resolved = UCineAssemblyNamingTokens::GetResolvedText(TokenString.Template, CustomizedCineAssembly);
-}
-
 TSharedRef<SWidget> FCineAssemblyCustomization::BuildProductionNameMenu()
 {
 	FMenuBuilder MenuBuilder(true, nullptr);
@@ -565,4 +496,64 @@ TSharedRef<SWidget> FCineAssemblyCustomization::BuildProductionNameMenu()
 	return MenuBuilder.MakeWidget();
 }
 
+FText FCineAssemblyCustomization::GetSubAssemblyTemplateName(int32 Index) const
+{
+	if (SubAssemblyNames.IsValidIndex(Index))
+	{
+		const FTemplateString& SubAssemblyName = SubAssemblyNames[Index];
+		return FText::FromString(FPaths::GetBaseFilename(SubAssemblyName.Template));
+	}
+	return FText::GetEmpty();
+}
+
+void FCineAssemblyCustomization::SetSubAssemblyTemplateName(int32 Index, const FText& InText, ETemplateStringType InTemplateStringType)
+{
+	if (SubAssemblyNames.IsValidIndex(Index))
+	{
+		FTemplateString& SubAssemblyName = SubAssemblyNames[Index];
+		if (FTemplateString* TemplateString = Algo::FindBy(CustomizedCineAssembly->SubAssemblyNames, SubAssemblyName.Template, &FTemplateString::Template))
+		{
+			const FString Path = FPaths::GetPath(SubAssemblyName.Template);
+
+			if (InTemplateStringType == ETemplateStringType::Template)
+			{
+				SubAssemblyName.Template = Path / InText.ToString();
+				TemplateString->Template = SubAssemblyName.Template;
+			}
+			else if (InTemplateStringType == ETemplateStringType::Resolved)
+			{
+				SubAssemblyName.Resolved = FText::FromString(Path / InText.ToString());
+				TemplateString->Resolved = SubAssemblyName.Resolved;
+			}
+		}
+	}
+}
+
+FText FCineAssemblyCustomization::GetMetadataTokenStringValue(const FString& InMetadataKey) const
+{
+	FTemplateString Value;
+	if (CustomizedCineAssembly->GetMetadataAsTokenString(InMetadataKey, Value))
+	{
+		return FText::FromString(Value.Template);
+	}
+	return FText::GetEmpty();
+}
+
+void FCineAssemblyCustomization::SetMetadataTokenStringValue(const FString& InMetadataKey, const FText& InValue, ETemplateStringType InTemplateStringType)
+{
+	FTemplateString TemplateString;
+	CustomizedCineAssembly->GetMetadataAsTokenString(InMetadataKey, TemplateString);
+
+	if (InTemplateStringType == ETemplateStringType::Template)
+	{
+		TemplateString.Template = InValue.ToString();
+	}
+	else if (InTemplateStringType == ETemplateStringType::Resolved)
+	{
+		TemplateString.Resolved = InValue;
+	}
+
+	CustomizedCineAssembly->SetMetadataAsTokenString(InMetadataKey, TemplateString);
+}
+
 #undef LOCTEXT_NAMESPACE
diff --git a/Engine/Plugins/VirtualProduction/CinematicAssemblyTools/Source/CineAssemblyToolsEditor/Private/CineAssemblyCustomization.h b/Engine/Plugins/VirtualProduction/CinematicAssemblyTools/Source/CineAssemblyToolsEditor/Private/CineAssemblyCustomization.h
--- a/Engine/Plugins/VirtualProduction/CinematicAssemblyTools/Source/CineAssemblyToolsEditor/Private/CineAssemblyCustomization.h
+++ b/Engine/Plugins/VirtualProduction/CinematicAssemblyTools/Source/CineAssemblyToolsEditor/Private/CineAssemblyCustomization.h
@@ -4,7 +4,9 @@
 
 #include "IDetailCustomization.h"
 
+#include "CineAssemblyNamingTokens.h"
 #include "LeadingZeroNumericTypeInterface.h"
+#include "NamingTokensEngineSubsystem.h"
 #include "Styling/SlateTypes.h"
 #include "UObject/TemplateString.h"
 
@@ -36,39 +38,37 @@ private:
 	*/
 	TSharedRef<SWidget> MakeStringValueWidget(const FAssemblyMetadataDesc& MetadataDesc);
 
-	/**
-	 * Evaluates all of the customization's template strings with the naming tokens subsystem.
-	 * This function is throttled to only run at a set frequency, to avoid the potential to constantly query the naming tokens subsystem.
-	 */
-	void EvaluateTokenStrings();
-
-	/**
-	 * Evaluates the input template string with the naming tokens subsystem.
-	 * This function is not throttled to allow for immediate updates.
-	 */
-	void EvaluateTokenString(FTemplateString& TokenString);
-
 	/** Checks if the customized CineAssembly contains a SubAssembly name in its list of assets to create */
 	ECheckBoxState IsSubAssemblyChecked(int32 Index) const;
 
 	/** Adds/Removes a SubAssembly name from the custmoized CineAssembly's list of SubAssemblies */
 	void SubAssemblyCheckStateChanged(ECheckBoxState CheckBoxState, int32 Index);
 
-	/** Returns the template text for a SubAssembly template name */
-	FText GetTemplateText(int32 Index) const;
-
-	/** Evaluates the token strings, then returns the resolved text for a SubAssembly template name */
-	FText GetResolvedText(int32 Index);
-
-	/** Modifies the template text of the SubAssembly template name and re-evaluates the token string to update the resolved text */
-	void OnTemplateTextCommitted(const FText& InText, ETextCommit::Type InCommitType, int32 Index);
-
 	/** Builds the drop-down menu list of productions */
 	TSharedRef<SWidget> BuildProductionNameMenu();
 
 	/** Determines whether the input asset should be filtered out of an object picker widget, based on whether it is of the input schema type */
 	bool ShouldFilterAssetBySchema(const FAssetData& InAssetData, FSoftObjectPath Schema);
 
+	/** Used to specify which member of an FTemplateString is desired */
+	enum class ETemplateStringType
+	{
+		Template,
+		Resolved
+	};
+
+	/** Returns the template text for a SubAssembly */
+	FText GetSubAssemblyTemplateName(int32 Index) const;
+
+	/** Updates either the template text or resolved text for a SubAssembly */
+	void SetSubAssemblyTemplateName(int32 Index, const FText& InText, ETemplateStringType InTemplateStringType);
+
+	/** Returns the template text for a token string metadata property */
+	FText GetMetadataTokenStringValue(const FString& InMetadataKey) const;
+
+	/** Updates either the template text or resolved text for a token string metadata property */
+	void SetMetadataTokenStringValue(const FString& InMetadataKey, const FText& InValue, ETemplateStringType InTemplateStringType);
+
 private:
 	/** The assembly being customized */
 	UCineAssembly* CustomizedCineAssembly;
@@ -81,4 +81,10 @@ private:
 
 	/** Type interface for a numeric entry box that supports leading zeroes */
 	TSharedPtr<FLeadingZeroNumericTypeInterface> LeadingZeroTypeInterface;
+
+	/** Context object allowing cine assembly tokens to evaluate */
+	TStrongObjectPtr<UCineAssemblyNamingTokensContext> NamingTokenContext;
+
+	/** Filter allowing cine assembly tokens to evaluate without the namespace prefix */
+	FNamingTokenFilterArgs FilterArgs;
 };
diff --git a/Engine/Plugins/VirtualProduction/CinematicAssemblyTools/Source/CineAssemblyToolsEditor/Private/CineAssemblyMetadataCustomization.cpp b/Engine/Plugins/VirtualProduction/CinematicAssemblyTools/Source/CineAssemblyToolsEditor/Private/CineAssemblyMetadataCustomization.cpp
--- a/Engine/Plugins/VirtualProduction/CinematicAssemblyTools/Source/CineAssemblyToolsEditor/Private/CineAssemblyMetadataCustomization.cpp
+++ b/Engine/Plugins/VirtualProduction/CinematicAssemblyTools/Source/CineAssemblyToolsEditor/Private/CineAssemblyMetadataCustomization.cpp
@@ -6,11 +6,12 @@
 #include "Algo/IndexOf.h"
 #include "CineAssembly.h"
 #include "CineAssemblyNamingTokens.h"
+#include "CineAssemblyToolsStyle.h"
 #include "IDetailChildrenBuilder.h"
 #include "IPropertyUtilities.h"
 #include "NamingTokensEngineSubsystem.h"
 #include "PropertyCustomizationHelpers.h"
-#include "STemplateStringEditableTextBox.h"
+#include "SNamingTokensEditableTextBox.h"
 #include "Widgets/Input/SCheckBox.h"
 #include "Widgets/Input/SEditableTextBox.h"
 #include "Widgets/Input/SMultiLineEditableTextBox.h"
@@ -74,8 +75,11 @@ TSharedRef<SWidget> FCineAssemblyMetadataCustomization::MakeStringDefaultValueWi
 
 	StringValueWidget->AddSlot(EvalTrueIndex)
 		[
-			SNew(STemplateStringEditableTextBox)
+			SNew(SNamingTokensEditableTextBox)
 				.AllowMultiLine(true)
+				.Style(FCineAssemblyToolsStyle::Get(), "DetailsEditableTextBox")
+				.ArgumentStyle(FCineAssemblyToolsStyle::Get(), "DetailsEditableTextBoxArguments")
+				.ShouldEvaluateTokens(false)
 				.Text_Lambda([&MetadataDesc]() -> FText
 					{
 						if (const FString* Value = MetadataDesc.DefaultValue.TryGet<FString>())
@@ -84,14 +88,6 @@ TSharedRef<SWidget> FCineAssemblyMetadataCustomization::MakeStringDefaultValueWi
 						}
 						return FText::GetEmpty();
 					})
-				.ResolvedText_Lambda([&MetadataDesc]()
-					{
-						if (const FString* Value = MetadataDesc.DefaultValue.TryGet<FString>())
-						{
-							return FText::FromString(*Value);
-						}
-						return FText::GetEmpty();
-					})
 				.OnTextCommitted_Lambda([&MetadataDesc](const FText& InText, ETextCommit::Type InCommitType)
 					{
 						MetadataDesc.DefaultValue.Set<FString>(InText.ToString());
diff --git a/Engine/Plugins/VirtualProduction/CinematicAssemblyTools/Source/CineAssemblyToolsEditor/Private/CineAssemblyMetadataCustomization.h b/Engine/Plugins/VirtualProduction/CinematicAssemblyTools/Source/CineAssemblyToolsEditor/Private/CineAssemblyMetadataCustomization.h
--- a/Engine/Plugins/VirtualProduction/CinematicAssemblyTools/Source/CineAssemblyToolsEditor/Private/CineAssemblyMetadataCustomization.h
+++ b/Engine/Plugins/VirtualProduction/CinematicAssemblyTools/Source/CineAssemblyToolsEditor/Private/CineAssemblyMetadataCustomization.h
@@ -21,7 +21,7 @@ public:
 
 private:
 	/**
-	* Given a FAssemblyMetadataDesc, makes a STemplateStringEditableTextBox for a string default value that has bEvaluateTokens set to true,
+	* Given a FAssemblyMetadataDesc, makes a SNamingTokensEditableTextBox for a string default value that has bEvaluateTokens set to true,
 	* and a SMultiLineEditableTextBox for a string default value that has bEvaluateTokens set to false.
 	*/
 	TSharedRef<SWidget> MakeStringDefaultValueWidget(FAssemblyMetadataDesc& MetadataDesc, TSharedRef<IPropertyHandle> PropertyHandle);
diff --git a/Engine/Plugins/VirtualProduction/CinematicAssemblyTools/Source/CineAssemblyToolsEditor/Private/CineAssemblySchemaCustomization.cpp b/Engine/Plugins/VirtualProduction/CinematicAssemblyTools/Source/CineAssemblyToolsEditor/Private/CineAssemblySchemaCustomization.cpp
--- a/Engine/Plugins/VirtualProduction/CinematicAssemblyTools/Source/CineAssemblyToolsEditor/Private/CineAssemblySchemaCustomization.cpp
+++ b/Engine/Plugins/VirtualProduction/CinematicAssemblyTools/Source/CineAssemblyToolsEditor/Private/CineAssemblySchemaCustomization.cpp
@@ -6,8 +6,10 @@
 #include "AssetRegistry/AssetRegistryModule.h"
 #include "AssetViewUtils.h"
 #include "CineAssemblySchema.h"
+#include "CineAssemblyToolsStyle.h"
 #include "DetailLayoutBuilder.h"
 #include "PropertyCustomizationHelpers.h"
+#include "SNamingTokensEditableTextBox.h"
 #include "Widgets/Input/SEditableTextBox.h"
 
 #define LOCTEXT_NAMESPACE "CineAssemblySchemaCustomization"
@@ -89,18 +91,20 @@ void FCineAssemblySchemaCustomization::CustomizeDetails(IDetailLayoutBuilder& De
 		]
 		.ValueContent()
 		[
-			SNew(SEditableTextBox)
-				.Font(IDetailLayoutBuilder::GetDetailFont())
+			SNew(SNamingTokensEditableTextBox)
+				.Style(FCineAssemblyToolsStyle::Get(), "DetailsEditableTextBox")
+				.ArgumentStyle(FCineAssemblyToolsStyle::Get(), "DetailsEditableTextBoxArguments")
+				.DisplayTokenIcon(false)
+				.ShouldEvaluateTokens(false)
 				.Text_Lambda([this]() -> FText
 					{
 						return FText::FromString(CustomizedSchema->DefaultAssemblyName);
 					})
 				.OnTextCommitted_Lambda([this](const FText& InText, ETextCommit::Type InCommitType)
 					{
-						CustomizedSchema->Modify();
 						CustomizedSchema->DefaultAssemblyName = InText.ToString();
 					})
-				.OnVerifyTextChanged_Raw(this, &FCineAssemblySchemaCustomization::ValidateDefaultAssemblyName)
+				.OnValidateTokenizedText_Raw(this, &FCineAssemblySchemaCustomization::ValidateDefaultAssemblyName)
 		];
 }
 
diff --git a/Engine/Plugins/VirtualProduction/CinematicAssemblyTools/Source/CineAssemblyToolsEditor/Private/CineAssemblyToolsStyle.cpp b/Engine/Plugins/VirtualProduction/CinematicAssemblyTools/Source/CineAssemblyToolsEditor/Private/CineAssemblyToolsStyle.cpp
--- a/Engine/Plugins/VirtualProduction/CinematicAssemblyTools/Source/CineAssemblyToolsEditor/Private/CineAssemblyToolsStyle.cpp
+++ b/Engine/Plugins/VirtualProduction/CinematicAssemblyTools/Source/CineAssemblyToolsEditor/Private/CineAssemblyToolsStyle.cpp
@@ -112,6 +112,19 @@ FCineAssemblyToolsStyle::FCineAssemblyToolsStyle()
 		.SetUniformPadding(FMargin(1))
 	);
 
+	// Textbox styles
+	FEditableTextBoxStyle DetailsEditableTextBoxStyle = FAppStyle::Get().GetWidgetStyle<FEditableTextBoxStyle>("NormalEditableTextBox");
+	DetailsEditableTextBoxStyle.SetFont(DEFAULT_FONT("Regular", 8));
+
+	Set("DetailsEditableTextBox", DetailsEditableTextBoxStyle);
+
+	FTextBlockStyle DetailsArgumentStyle = FAppStyle::Get().GetWidgetStyle<FTextBlockStyle>(TEXT("NormalText"));
+	DetailsArgumentStyle
+		.SetFont(DEFAULT_FONT("Bold", 8))
+		.SetColorAndOpacity(FLinearColor(1.0f, 1.0f, 1.0f));
+
+	Set("DetailsEditableTextBoxArguments", DetailsArgumentStyle);
+
 	FSlateStyleRegistry::RegisterSlateStyle(*this);
 }
 
diff --git a/Engine/Plugins/VirtualProduction/CinematicAssemblyTools/Source/CineAssemblyToolsEditor/Private/UI/CineAssembly/SCineAssemblyConfigPanel.cpp b/Engine/Plugins/VirtualProduction/CinematicAssemblyTools/Source/CineAssemblyToolsEditor/Private/UI/CineAssembly/SCineAssemblyConfigPanel.cpp
--- a/Engine/Plugins/VirtualProduction/CinematicAssemblyTools/Source/CineAssemblyToolsEditor/Private/UI/CineAssembly/SCineAssemblyConfigPanel.cpp
+++ b/Engine/Plugins/VirtualProduction/CinematicAssemblyTools/Source/CineAssemblyToolsEditor/Private/UI/CineAssembly/SCineAssemblyConfigPanel.cpp
@@ -3,12 +3,11 @@
 #include "SCineAssemblyConfigPanel.h"
 
 #include "AssetToolsModule.h"
-#include "CineAssemblyNamingTokens.h"
 #include "CineAssemblyToolsStyle.h"
 #include "IStructureDetailsView.h"
 #include "Modules/ModuleManager.h"
 #include "PropertyEditorModule.h"
-#include "STemplateStringEditableTextBox.h"
+#include "SNamingTokensEditableTextBox.h"
 #include "Widgets/Input/SSegmentedControl.h"
 
 #define LOCTEXT_NAMESPACE "SCineAssemblyConfigPanel"
@@ -33,6 +32,11 @@ void SCineAssemblyConfigPanel::Construct(const FArguments& InArgs, UCineAssembly
 		DetailsView->SetObject(CineAssemblyToConfigure, true);
 	}
 
+	NamingTokenContext = TStrongObjectPtr<UCineAssemblyNamingTokensContext>(NewObject<UCineAssemblyNamingTokensContext>());
+	NamingTokenContext->Assembly = InAssembly;
+
+	FilterArgs.AdditionalNamespacesToInclude.Add(UCineAssemblyNamingTokens::TokenNamespace);
+
 		ChildSlot
 		[
 			SNew(SVerticalBox)
@@ -88,22 +92,24 @@ void SCineAssemblyConfigPanel::Construct(const FArguments& InArgs, UCineAssembly
 								.FillWidth(1.0f)
 								.VAlign(VAlign_Center)
 								[
-									SNew(STemplateStringEditableTextBox)
-										.Text_Lambda([this]() -> FText
+									SNew(SNamingTokensEditableTextBox)
+										.Contexts({ NamingTokenContext.Get() })
+										.FilterArgs(FilterArgs)
+										.EvaluationFrequency(1.0f)
+										.OnValidateTokenizedText(this, &SCineAssemblyConfigPanel::ValidateAssemblyName)
+										.OnValidateResolvedText(this, &SCineAssemblyConfigPanel::ValidateAssemblyName)
+										.Text_Lambda([&Assembly = CineAssemblyToConfigure]() -> FText
 											{
-												return FText::FromString(CineAssemblyToConfigure->AssemblyName.Template);
+												return FText::FromString(Assembly->AssemblyName.Template);
 											})
-										.ResolvedText_Lambda([this]()
+										.OnTextCommitted_Lambda([&Assembly = CineAssemblyToConfigure](const FText& InText, ETextCommit::Type InCommitType)
 											{
-												EvaluateTokenString(CineAssemblyToConfigure->AssemblyName);
-												return CineAssemblyToConfigure->AssemblyName.Resolved;
+												Assembly->AssemblyName.Template = InText.ToString();
 											})
-										.OnTextCommitted_Lambda([this](const FText& InText, ETextCommit::Type InCommitType)
+										.OnTokenizedTextEvaluated_Lambda([&Assembly = CineAssemblyToConfigure](const FText& InText)
 											{
-												CineAssemblyToConfigure->Modify();
-												CineAssemblyToConfigure->AssemblyName.Template = InText.ToString();
+												Assembly->AssemblyName.Resolved = InText;
 											})
-										.OnValidateTokenizedText(this, &SCineAssemblyConfigPanel::ValidateAssemblyName)
 								]
 						]
 				]
@@ -145,16 +151,6 @@ void SCineAssemblyConfigPanel::Refresh()
 	PopulateHierarchyTree();
 }
 
-void SCineAssemblyConfigPanel::EvaluateTokenString(FTemplateString& StringToEvaluate)
-{
-	FDateTime CurrentTime = FDateTime::Now();
-	if ((CurrentTime - LastTokenUpdateTime).GetSeconds() >= 1.0f)
-	{
-		StringToEvaluate.Resolved = UCineAssemblyNamingTokens::GetResolvedText(StringToEvaluate.Template, CineAssemblyToConfigure);
-		LastTokenUpdateTime = CurrentTime;
-	}
-}
-
 TSharedRef<SWidget> SCineAssemblyConfigPanel::MakeDetailsWidget()
 {
 	FPropertyEditorModule& PropertyEditor = FModuleManager::Get().LoadModuleChecked<FPropertyEditorModule>(TEXT("PropertyEditor"));
@@ -287,17 +283,6 @@ TSharedRef<SWidget> SCineAssemblyConfigPanel::MakeHierarchyWidget()
 
 	PopulateHierarchyTree();
 
-	// Register a Slate timer that runs at a set frequency to evaluate all of the tokens in the tree view.
-	// This will automatically be unregistered when this window is destroyed.
-	constexpr float TimerFrequency = 1.0f;
-	RegisterActiveTimer(TimerFrequency, FWidgetActiveTimerDelegate::CreateLambda([this](double InCurrentTime, float InDeltaTime) -> EActiveTimerReturnType
-		{
-			EvaluateHierarchyTokensRecursive(RootItem);
-			HierarchyTreeView->RequestTreeRefresh();
-
-			return EActiveTimerReturnType::Continue;
-		}));
-
 	return SNew(SVerticalBox)
 
 		+ SVerticalBox::Slot()
@@ -380,40 +365,12 @@ void SCineAssemblyConfigPanel::PopulateHierarchyTree()
 
 			ParentItem->ChildAssets.Add(NewItem);
 		}
-
-		EvaluateHierarchyTokensRecursive(RootItem);
 	}
 
 	HierarchyTreeView->RequestTreeRefresh();
 	ExpandTreeRecursive(RootItem);
 }
 
-void SCineAssemblyConfigPanel::EvaluateHierarchyTokensRecursive(TSharedPtr<FHierarchyTreeItem> TreeItem)
-{
-	// Evaluate the token template string for this tree item
-	TreeItem->Path.Resolved = UCineAssemblyNamingTokens::GetResolvedText(TreeItem->Path.Template, CineAssemblyToConfigure);
-	if (TreeItem->Path.Resolved.IsEmpty())
-	{
-		TreeItem->Path.Resolved = FText::FromString(TreeItem->Path.Template);
-	}
-
-	// Evaluate the tokens for all of the child assets, then resort them alphabetically based on the resolved paths
-	for (const TSharedPtr<FHierarchyTreeItem>& Asset : TreeItem->ChildAssets)
-	{
-		EvaluateHierarchyTokensRecursive(Asset);
-	}
-
-	TreeItem->ChildAssets.Sort([](const TSharedPtr<FHierarchyTreeItem>& A, const TSharedPtr<FHierarchyTreeItem>& B) { return A->Path.Resolved.ToString() < B->Path.Resolved.ToString(); });
-
-	// Evaluate the tokens for all of the child folders, then resort them alphabetically based on the resolved paths
-	for (const TSharedPtr<FHierarchyTreeItem>& Child : TreeItem->ChildFolders)
-	{
-		EvaluateHierarchyTokensRecursive(Child);
-	}
-
-	TreeItem->ChildFolders.Sort([](const TSharedPtr<FHierarchyTreeItem>& A, const TSharedPtr<FHierarchyTreeItem>& B) { return A->Path.Resolved.ToString() < B->Path.Resolved.ToString(); });
-}
-
 void SCineAssemblyConfigPanel::ExpandTreeRecursive(TSharedPtr<FHierarchyTreeItem> TreeItem) const
 {
 	if (HierarchyTreeView)
@@ -477,10 +434,23 @@ TSharedRef<ITableRow> SCineAssemblyConfigPanel::OnGenerateTreeRow(TSharedPtr<FHi
 			+ SHorizontalBox::Slot()
 				.VAlign(VAlign_Center)
 				[
-					SNew(STextBlock)
+					SNew(SNamingTokensEditableTextBox)
+						.Contexts({ NamingTokenContext.Get() })
+						.FilterArgs(FilterArgs)
+						.EvaluationFrequency(1.0f)
+						.IsReadOnly(true)
+						.DisplayBorderImage(false)
+						.DisplayTokenIcon(false)
 						.Text_Lambda([this, TreeItem]() -> FText
 							{
-								return (TreeItem == RootItem) ? LOCTEXT("RootPathName", "Root Folder") : FText::FromString(FPaths::GetPathLeaf(TreeItem->Path.Resolved.ToString()));
+								return (TreeItem == RootItem) ? LOCTEXT("RootPathName", "Root Folder") : FText::FromString(FPaths::GetPathLeaf(TreeItem->Path.Template));
+							})
+						.OnTokenizedTextEvaluated_Lambda([this, TreeItem](const FText& InText)
+							{
+								if (TreeItem != RootItem)
+								{
+									TreeItem->Path.Resolved = InText;
+								}
 							})
 				]
 		];
diff --git a/Engine/Plugins/VirtualProduction/CinematicAssemblyTools/Source/CineAssemblyToolsEditor/Private/UI/CineAssembly/SCineAssemblyConfigPanel.h b/Engine/Plugins/VirtualProduction/CinematicAssemblyTools/Source/CineAssemblyToolsEditor/Private/UI/CineAssembly/SCineAssemblyConfigPanel.h
--- a/Engine/Plugins/VirtualProduction/CinematicAssemblyTools/Source/CineAssemblyToolsEditor/Private/UI/CineAssembly/SCineAssemblyConfigPanel.h
+++ b/Engine/Plugins/VirtualProduction/CinematicAssemblyTools/Source/CineAssemblyToolsEditor/Private/UI/CineAssembly/SCineAssemblyConfigPanel.h
@@ -5,7 +5,9 @@
 #include "Widgets/SCompoundWidget.h"
 
 #include "CineAssembly.h"
+#include "CineAssemblyNamingTokens.h"
 #include "IDetailsView.h"
+#include "NamingTokensEngineSubsystem.h"
 #include "Widgets/Layout/SWidgetSwitcher.h"
 #include "Widgets/Views/STreeView.h"
 
@@ -79,18 +81,9 @@ private:
 	/** Recursively expands every item in the tree view */
 	void ExpandTreeRecursive(TSharedPtr<FHierarchyTreeItem> TreeItem) const;
 
-	/** Recursively evaluates the tokens in the path of each item in the tree view */
-	void EvaluateHierarchyTokensRecursive(TSharedPtr<FHierarchyTreeItem> TreeItem);
-
 	/** Validates the user input text for the assembly name */
 	bool ValidateAssemblyName(const FText& InText, FText& OutErrorMessage) const;
 
-	/**
-	 * Evaluates the input template string with the naming tokens subsystem, and stores the result in the Resolved text.
-	 * This function is throttled to only run at a set frequency, to avoid the potential to constantly query the naming tokens subsystem.
-	 */
-	void EvaluateTokenString(FTemplateString& StringToEvaluate);
-
 	/** Filter used by the Details View to determine which custom rows to show */
 	bool IsCustomRowVisible(FName RowName, FName ParentName);
 
@@ -116,6 +109,9 @@ private:
 	/** A read-only tree view of folders and assets that will be created based on the selected schema */
 	TSharedPtr<STreeView<TSharedPtr<FHierarchyTreeItem>>> HierarchyTreeView;
 
-	/** The last time the naming tokens were updated */
-	FDateTime LastTokenUpdateTime;
+	/** Naming Token Context object, with a pointer to the assembly being configured */
+	TStrongObjectPtr<UCineAssemblyNamingTokensContext> NamingTokenContext;
+
+	/** Naming Token filter properties */
+	FNamingTokenFilterArgs FilterArgs;
 };
diff --git a/Engine/Plugins/VirtualProduction/CinematicAssemblyTools/Source/CineAssemblyToolsEditor/Private/UI/CineAssembly/SCineAssemblySchemaWindow.cpp b/Engine/Plugins/VirtualProduction/CinematicAssemblyTools/Source/CineAssemblyToolsEditor/Private/UI/CineAssembly/SCineAssemblySchemaWindow.cpp
--- a/Engine/Plugins/VirtualProduction/CinematicAssemblyTools/Source/CineAssemblyToolsEditor/Private/UI/CineAssembly/SCineAssemblySchemaWindow.cpp
+++ b/Engine/Plugins/VirtualProduction/CinematicAssemblyTools/Source/CineAssemblyToolsEditor/Private/UI/CineAssembly/SCineAssemblySchemaWindow.cpp
@@ -739,16 +739,18 @@ TSharedRef<ITableRow> SCineAssemblySchemaWindow::OnGenerateAssetRow(TSharedPtr<F
 		.ToolTipText(LOCTEXT("AssetRowTooltip", "Set the name and template subsequence (optional) for this subsequence, then add it to the content tree on the right."));
 
 	// Store a reference to the editable textblock in the row to easily set it to edit mode for renaming
-	Row->TextBlock = SNew(SInlineEditableTextBlock)
-		.Text_Lambda([InItem]() 
+	Row->TextBlock = SNew(SNamingTokensEditableTextBox)
+		.SupportInlineEdit(true)
+		.ShouldEvaluateTokens(false)
+		.Text_Lambda([InItem]()
 		{ 
 			return FText::FromString(InItem->Name); 
 		})
-		.OnVerifyTextChanged(this, &SCineAssemblySchemaWindow::ValidateAssetName)
-		.OnTextCommitted_Lambda([this, InItem](const FText& InText, ETextCommit::Type InCommitType) 
-		{ 
+		.OnTextCommitted_Lambda([this, InItem](const FText& InText, ETextCommit::Type InCommitType)
+		{
 			InItem->Name = InText.ToString();
-		});
+		})
+		.OnValidateTokenizedText(this, &SCineAssemblySchemaWindow::ValidateAssetName);
 
 	Row->SetContent(SNew(SBorder)
 		.Padding(8.0f)
@@ -890,9 +892,11 @@ TSharedRef<ITableRow> SCineAssemblySchemaWindow::OnGenerateTreeRow(TSharedPtr<FS
 		Icon->SetColorAndOpacity(FLinearColor::White);
 	}
 
+	const bool bIsRootItem = (TreeItem == RootItem);
+
 	return SNew(STableRow<TSharedPtr<FSchemaTreeItem>>, OwnerTable)
 		.ShowSelection(true)
-		.Padding(FMargin(8.0f, 2.0f, 8.0f, 0.0f))
+		.Padding(FMargin(8.0f, 0.0f, 8.0f, 0.0f))
 		.OnCanAcceptDrop_Lambda([](const FDragDropEvent& InDragDropEvent, EItemDropZone InDropZone, TSharedPtr<FSchemaTreeItem> InItem)
 			{
 				// Only folder items can accept drops, and only onto to th item (not above or below)
@@ -915,16 +919,19 @@ TSharedRef<ITableRow> SCineAssemblySchemaWindow::OnGenerateTreeRow(TSharedPtr<FS
 			+ SHorizontalBox::Slot()
 				.VAlign(VAlign_Center)
 				[
-					SAssignNew(TreeItem->NameWidget, SInlineEditableTextBlock)
+					SAssignNew(TreeItem->NameWidget, SNamingTokensEditableTextBox)
+						.SupportInlineEdit(true)
+						.ShouldEvaluateTokens(false)
+						.TextBoxPadding(FMargin(4, 2, 0, 2))
+						.IsReadOnly(bIsRootItem)
+						.DisplayTokenIcon(false)
+						.DisplayBorderImage(false)
+						.DisplayErrorMessage(!bIsRootItem)
 						.Text_Lambda([this, TreeItem]() -> FText
 							{
 								return (TreeItem == RootItem) ? LOCTEXT("RootPathName", "Root Folder") : FText::FromString(FPaths::GetPathLeaf(TreeItem->Path));
 							})
-						.IsReadOnly_Lambda([this, TreeItem]() -> bool
-							{
-								return (TreeItem == RootItem);
-							})
-						.OnVerifyTextChanged_Lambda([this, TreeItem](const FText& InText, FText& OutErrorMessage) -> bool
+						.OnValidateTokenizedText_Lambda([this, TreeItem](const FText& InText, FText& OutErrorMessage) -> bool
 							{
 								if (TreeItem->Type == FSchemaTreeItem::EItemType::Folder)
 								{
@@ -971,7 +978,6 @@ void SCineAssemblySchemaWindow::OnTreeItemsRebuilt()
 	// Upon regenerating the tree view, allow the user to immediately interact with the name widget of the newly added folder in order to rename it
 	if (MostRecentlyAddedItem && MostRecentlyAddedItem->NameWidget)
 	{
-		FSlateApplication::Get().SetKeyboardFocus(MostRecentlyAddedItem->NameWidget.ToSharedRef());
 		MostRecentlyAddedItem->NameWidget->EnterEditingMode();
 		MostRecentlyAddedItem.Reset();
 	}
@@ -1009,7 +1015,6 @@ TSharedPtr<SWidget> SCineAssemblySchemaWindow::MakeContentTreeContextMenu()
 				FUIAction(
 					FExecuteAction::CreateLambda([SelectedTreeItem]()
 						{
-							FSlateApplication::Get().SetKeyboardFocus(SelectedTreeItem->NameWidget.ToSharedRef());
 							SelectedTreeItem->NameWidget->EnterEditingMode();
 						})),
 				NAME_None,
@@ -1037,9 +1042,8 @@ void SCineAssemblySchemaWindow::OnTreeViewDoubleClick(TSharedPtr<FSchemaTreeItem
 	// Get the row for the input item and put its textblock in edit mode so the user can rename the item
 	if (TreeItem && TreeItem->NameWidget)
 	{
-		if (TreeItem != TopLevelAssemblyItem)
+		if (TreeItem != RootItem && TreeItem != TopLevelAssemblyItem)
 		{
-			FSlateApplication::Get().SetKeyboardFocus(TreeItem->NameWidget.ToSharedRef());
 			TreeItem->NameWidget->EnterEditingMode();
 		}
 	}
@@ -1207,7 +1211,6 @@ void SCineAssemblySchemaWindow::EnterEditMode(TSharedPtr<FSchemaAssetParameters>
 	{
 		if (Widget->TextBlock)
 		{
-			FSlateApplication::Get().SetKeyboardFocus(Widget->TextBlock.ToSharedRef());
 			Widget->TextBlock->EnterEditingMode();
 		}
 	}
diff --git a/Engine/Plugins/VirtualProduction/CinematicAssemblyTools/Source/CineAssemblyToolsEditor/Private/UI/CineAssembly/SCineAssemblySchemaWindow.h b/Engine/Plugins/VirtualProduction/CinematicAssemblyTools/Source/CineAssemblyToolsEditor/Private/UI/CineAssembly/SCineAssemblySchemaWindow.h
--- a/Engine/Plugins/VirtualProduction/CinematicAssemblyTools/Source/CineAssemblyToolsEditor/Private/UI/CineAssembly/SCineAssemblySchemaWindow.h
+++ b/Engine/Plugins/VirtualProduction/CinematicAssemblyTools/Source/CineAssemblyToolsEditor/Private/UI/CineAssembly/SCineAssemblySchemaWindow.h
@@ -6,8 +6,8 @@
 
 #include "CineAssemblySchema.h"
 #include "Input/DragAndDrop.h"
+#include "SNamingTokensEditableTextBox.h"
 #include "Widgets/Layout/SWidgetSwitcher.h"
-#include "Widgets/Text/SInlineEditableTextBlock.h"
 #include "Widgets/Views/SListView.h"
 #include "Widgets/Views/STreeView.h"
 
@@ -35,7 +35,7 @@ class SSchemaAssetTableRow : public STableRow<TSharedPtr<FSchemaAssetParameters>
 {
 public:
 	/** Editable text widget stored in this row to easily trigger edit mode */
-	TSharedPtr<SInlineEditableTextBlock> TextBlock;
+	TSharedPtr<SNamingTokensEditableTextBox> TextBlock;
 };
 
 /** An entry in the content folder tree view with knowledge of its path, children, and parent */
@@ -67,7 +67,7 @@ struct FSchemaTreeItem
 	FSoftObjectPath AssetTemplate;
 
 	/** The text widget that displays this item's name and supports renaming */
-	TSharedPtr<SInlineEditableTextBlock> NameWidget;
+	TSharedPtr<SNamingTokensEditableTextBox> NameWidget;
 };
 
 /** Drag and Drop operation to handle dragging schema assets and folders  */
diff --git a/Engine/Plugins/VirtualProduction/CinematicAssemblyTools/Source/CineAssemblyToolsEditor/Private/UI/CineAssembly/SDuplicateAssemblyWindow.cpp b/Engine/Plugins/VirtualProduction/CinematicAssemblyTools/Source/CineAssemblyToolsEditor/Private/UI/CineAssembly/SDuplicateAssemblyWindow.cpp
--- a/Engine/Plugins/VirtualProduction/CinematicAssemblyTools/Source/CineAssemblyToolsEditor/Private/UI/CineAssembly/SDuplicateAssemblyWindow.cpp
+++ b/Engine/Plugins/VirtualProduction/CinematicAssemblyTools/Source/CineAssemblyToolsEditor/Private/UI/CineAssembly/SDuplicateAssemblyWindow.cpp
@@ -14,7 +14,7 @@
 #include "ObjectTools.h"
 #include "SCineAssemblyConfigPanel.h"
 #include "Sections/MovieSceneSubSection.h"
-#include "STemplateStringEditableTextBox.h"
+#include "SNamingTokensEditableTextBox.h"
 #include "Tracks/MovieSceneSubTrack.h"
 #include "Widgets/Input/SEditableTextBox.h"
 #include "Widgets/Layout/SSeparator.h"
@@ -145,25 +145,6 @@ TSharedRef<SWidget> SDuplicateAssemblyWindow::MakeSubsequencesPanel()
 
 	TSharedRef<SWidget> PathPicker = ContentBrowserModule.Get().CreatePathPicker(PathPickerConfig);
 
-	// Register a Slate timer that runs at a set frequency to evaluate all of the tokens in the tree view.
-	// This will automatically be unregistered when this window is destroyed.
-	constexpr float TimerFrequency = 1.0f;
-	RegisterActiveTimer(TimerFrequency, FWidgetActiveTimerDelegate::CreateLambda([this](double InCurrentTime, float InDeltaTime) -> EActiveTimerReturnType
-		{
-			for (TSharedPtr<FSequenceTreeItem>& Item : TreeItems)
-			{
-				ResolveSubsequenceNameRecursive(Item);
-			}
-
-			return EActiveTimerReturnType::Continue;
-		}));
-
-	// Resolve the subsequence names once, immediately, so they appear correctly when the window first opens
-	for (TSharedPtr<FSequenceTreeItem>& Item : TreeItems)
-	{
-		ResolveSubsequenceNameRecursive(Item);
-	}
-
 	return SNew(SBorder)
 		.BorderImage(FCineAssemblyToolsStyle::Get().GetBrush("ProductionWizard.PanelBackground"))
 		.Padding(16.0f)
@@ -317,7 +298,7 @@ FReply SDuplicateAssemblyWindow::OnCancelClicked()
 
 TSharedRef<ITableRow> SDuplicateAssemblyWindow::OnGenerateTreeRow(TSharedPtr<FSequenceTreeItem> TreeItem, const TSharedRef<STableViewBase>& OwnerTable)
 {
-	return SNew(SSubsequenceDuplicationRow, OwnerTable, TreeItem);
+	return SNew(SSubsequenceDuplicationRow, OwnerTable, TreeItem, DuplicateAssembly.Get());
 }
 
 void SDuplicateAssemblyWindow::OnGetChildren(TSharedPtr<FSequenceTreeItem> TreeItem, TArray<TSharedPtr<FSequenceTreeItem>>& OutNodes)
@@ -539,18 +520,6 @@ void SDuplicateAssemblyWindow::BuildDuplicationMapRecursive(const TSharedPtr<FSe
 	}
 }
 
-void SDuplicateAssemblyWindow::ResolveSubsequenceNameRecursive(TSharedPtr<FSequenceTreeItem> InItem)
-{
-	// Evaluate the token template string for this tree item
-	FSubsequenceDuplicationData& DuplicationData = InItem->DuplicationData;
-	DuplicationData.SubsequenceName.Resolved = UCineAssemblyNamingTokens::GetResolvedText(DuplicationData.SubsequenceName.Template, DuplicateAssembly.Get());
-
-	for (const TSharedPtr<FSequenceTreeItem>& Child : InItem->Children)
-	{
-		ResolveSubsequenceNameRecursive(Child);
-	}
-}
-
 void SDuplicateAssemblyWindow::ExpandTreeItem(TSharedPtr<FSequenceTreeItem> InItem) const
 {
 	TreeView->SetItemExpansion(InItem, true);
@@ -603,10 +572,15 @@ ECineAssemblyDuplicationMode SDuplicateAssemblyWindow::GetDuplicationPreference(
 	return ECineAssemblyDuplicationMode::DuplicateOriginal;
 }
 
-void SSubsequenceDuplicationRow::Construct(const FArguments& Args, const TSharedRef<STableViewBase>& OwnerTableView, const TSharedPtr<FSequenceTreeItem>& InTreeItem)
+void SSubsequenceDuplicationRow::Construct(const FArguments& Args, const TSharedRef<STableViewBase>& OwnerTableView, const TSharedPtr<FSequenceTreeItem>& InTreeItem, UCineAssembly* DuplicateAssembly)
 {
 	TreeItem = InTreeItem;
 
+	NamingTokenContext = TStrongObjectPtr<UCineAssemblyNamingTokensContext>(NewObject<UCineAssemblyNamingTokensContext>());
+	NamingTokenContext->Assembly = DuplicateAssembly;
+
+	FilterArgs.AdditionalNamespacesToInclude.Add(UCineAssemblyNamingTokens::TokenNamespace);
+
 	FSuperRowType::FArguments StyleArguments = FSuperRowType::FArguments()
 		.Padding(FMargin(4.0f, 8.0f))
 		.Style(&FAppStyle().GetWidgetStyle<FTableRowStyle>("TableView.AlternatingRow"));
@@ -715,36 +689,31 @@ TSharedRef<SWidget> SSubsequenceDuplicationRow::MakeDuplicateNameWidget()
 		.FillWidth(1.0f)
 		.Padding(10.0f, 0.0f, 16.0f, 0.0f)
 		[
-			// TODO: STemplateStringEditableTextBox does not currently support input validation because SMultiLineEditableText does not support it.
-			// Once we add support to that widget, add some input validation here to ensure that the input is a valid object name
-			SNew(STemplateStringEditableTextBox)
-				.Text_Lambda([this]()
+			SNew(SNamingTokensEditableTextBox)
+				.Contexts({ NamingTokenContext.Get() })
+				.FilterArgs(FilterArgs)
+				.EvaluationFrequency(1.0f)
+				.IsReadOnly_Lambda([this]()
 					{
-						if (TreeItem->DuplicationData.DuplicationMode == ECineAssemblyDuplicationMode::DuplicateOriginal)
-						{
-							const FString SubsequenceName = FPaths::GetBaseFilename(TreeItem->DuplicationData.SubsequenceName.Template);
-							return FText::FromString(SubsequenceName);
-						}
-						return TreeItem->Subsequence->GetDisplayName();
+						return TreeItem->DuplicationData.DuplicationMode != ECineAssemblyDuplicationMode::DuplicateOriginal;
+					})
+				.Visibility_Lambda([this]()
+					{
+						return (TreeItem->DuplicationData.DuplicationMode != ECineAssemblyDuplicationMode::Remove) ? EVisibility::Visible : EVisibility::Hidden;
 					})
-				.ResolvedText_Lambda([this]()
+				.Text_Lambda([this]()
 					{
 						if (TreeItem->DuplicationData.DuplicationMode == ECineAssemblyDuplicationMode::DuplicateOriginal)
 						{
-							const FString SubsequenceName = FPaths::GetBaseFilename(TreeItem->DuplicationData.SubsequenceName.Resolved.ToString());
+							const FString SubsequenceName = FPaths::GetBaseFilename(TreeItem->DuplicationData.SubsequenceName.Template);
 							return FText::FromString(SubsequenceName);
 						}
 						return TreeItem->Subsequence->GetDisplayName();
 					})
 				.OnTextCommitted(this, &SSubsequenceDuplicationRow::OnSubsequenceNameCommitted)
-				.IsReadOnly_Lambda([this]()
-					{
-						return TreeItem->DuplicationData.DuplicationMode != ECineAssemblyDuplicationMode::DuplicateOriginal;
-					})
-				.Visibility_Lambda([this]()
-					{
-						return (TreeItem->DuplicationData.DuplicationMode != ECineAssemblyDuplicationMode::Remove) ? EVisibility::Visible : EVisibility::Hidden;
-					})
+				.OnTokenizedTextEvaluated(this, &SSubsequenceDuplicationRow::OnSubsequenceTokenTextEvaluated)
+				.OnValidateTokenizedText(this, &SSubsequenceDuplicationRow::ValidateSubsequenceName)
+				.OnValidateResolvedText(this, &SSubsequenceDuplicationRow::ValidateSubsequenceName)
 		];
 }
 
@@ -778,9 +747,41 @@ void SSubsequenceDuplicationRow::SetDuplicationModeRecursive(TSharedPtr<FSequenc
 
 void SSubsequenceDuplicationRow::OnSubsequenceNameCommitted(const FText& InText, ETextCommit::Type InCommitType)
 {
-	const FString OldPath = TreeItem->DuplicationData.SubsequenceName.Template;
-	const FString NewPath = FPaths::GetPath(OldPath) / InText.ToString();
-	TreeItem->DuplicationData.SubsequenceName.Template = NewPath;
+	if (TreeItem->DuplicationData.DuplicationMode == ECineAssemblyDuplicationMode::DuplicateOriginal)
+	{
+		const FString OldPath = TreeItem->DuplicationData.SubsequenceName.Template;
+		const FString NewPath = FPaths::GetPath(OldPath) / InText.ToString();
+		TreeItem->DuplicationData.SubsequenceName.Template = NewPath;
+	}
+}
+
+void SSubsequenceDuplicationRow::OnSubsequenceTokenTextEvaluated(const FText& InText)
+{
+	if (TreeItem->DuplicationData.DuplicationMode == ECineAssemblyDuplicationMode::DuplicateOriginal)
+	{
+		const FString OldPath = TreeItem->DuplicationData.SubsequenceName.Template;
+		const FString NewPath = FPaths::GetPath(OldPath) / InText.ToString();
+		TreeItem->DuplicationData.SubsequenceName.Resolved = FText::FromString(NewPath);
+	}
+}
+
+bool SSubsequenceDuplicationRow::ValidateSubsequenceName(const FText& InText, FText& OutErrorMessage) const
+{
+	// Ensure that the name does not contain any characters that would be invalid for an asset name
+	// This matches the validation that would happen if the user was renaming an asset in the content browser
+	FString InvalidCharacters = INVALID_OBJECTNAME_CHARACTERS INVALID_LONGPACKAGE_CHARACTERS;
+
+	// These characters are actually valid, because we want to support naming tokens
+	InvalidCharacters = InvalidCharacters.Replace(TEXT("{}"), TEXT(""));
+	InvalidCharacters = InvalidCharacters.Replace(TEXT(":"), TEXT(""));
+
+	const FString PotentialName = InText.ToString();
+	if (!FName::IsValidXName(PotentialName, InvalidCharacters, &OutErrorMessage))
+	{
+		return false;
+	}
+
+	return true;
 }
 
 #undef LOCTEXT_NAMESPACE
diff --git a/Engine/Plugins/VirtualProduction/CinematicAssemblyTools/Source/CineAssemblyToolsEditor/Private/UI/CineAssembly/SDuplicateAssemblyWindow.h b/Engine/Plugins/VirtualProduction/CinematicAssemblyTools/Source/CineAssemblyToolsEditor/Private/UI/CineAssembly/SDuplicateAssemblyWindow.h
--- a/Engine/Plugins/VirtualProduction/CinematicAssemblyTools/Source/CineAssemblyToolsEditor/Private/UI/CineAssembly/SDuplicateAssemblyWindow.h
+++ b/Engine/Plugins/VirtualProduction/CinematicAssemblyTools/Source/CineAssemblyToolsEditor/Private/UI/CineAssembly/SDuplicateAssemblyWindow.h
@@ -5,7 +5,9 @@
 #include "Widgets/SWindow.h"
 
 #include "CineAssembly.h"
+#include "CineAssemblyNamingTokens.h"
 #include "MovieSceneSequence.h"
+#include "NamingTokensEngineSubsystem.h"
 #include "Widgets/Views/STreeView.h"
 
 class UMovieSceneFolder;
@@ -55,7 +57,7 @@ public:
 	SLATE_BEGIN_ARGS(SSubsequenceDuplicationRow) {}
 	SLATE_END_ARGS()
 
-	void Construct(const FArguments& Args, const TSharedRef<STableViewBase>& OwnerTableView, const TSharedPtr<FSequenceTreeItem>& InTreeItem);
+	void Construct(const FArguments& Args, const TSharedRef<STableViewBase>& OwnerTableView, const TSharedPtr<FSequenceTreeItem>& InTreeItem, UCineAssembly* DuplicateAssembly);
 
 	/** Creates the widget for this row for the specified column */
 	virtual TSharedRef<SWidget> GenerateWidgetForColumn(const FName& ColumnName) override;
@@ -69,15 +71,27 @@ private:
 	/** Updates the duplication mode for the input tree item, and potentially its children */
 	void SetDuplicationModeRecursive(TSharedPtr<FSequenceTreeItem> InItem, ECineAssemblyDuplicationMode Mode);
 
-	/** Updates the subsequence name associated with this tree row */
+	/** Updates the subsequence template name associated with this tree row */
 	void OnSubsequenceNameCommitted(const FText& InText, ETextCommit::Type InCommitType);
 
+	/** Updates the subsequence resolved name after the token text was evaluated */
+	void OnSubsequenceTokenTextEvaluated(const FText& InText);
+
+	/** Validates the user input text for the subsequence name */
+	bool ValidateSubsequenceName(const FText& InText, FText& OutErrorMessage) const;
+
 	/** Get the display name of the input duplication mode */
 	FText GetDuplicationModeDisplayName(ECineAssemblyDuplicationMode DuplicationMode) const;
 
 private:
 	/** The tree view item displayed by this row widget */
 	TSharedPtr<FSequenceTreeItem> TreeItem;
+
+	/** Naming Token Context object, with a pointer to the duplicate assembly */
+	TStrongObjectPtr<UCineAssemblyNamingTokensContext> NamingTokenContext;
+
+	/** Naming Token filter properties */
+	FNamingTokenFilterArgs FilterArgs;
 };
 
 /** A window to configure how a Cine Assembly asset should be duplicated */
@@ -130,9 +144,6 @@ private:
 	/** Converts the subsequence tree into a map of subsequences to their associated duplication data */
 	void BuildDuplicationMapRecursive(const TSharedPtr<FSequenceTreeItem> InTreeItem);
 
-	/** Resolves any tokens present in the subsequence name of the input tree item, then resolves its children */
-	void ResolveSubsequenceNameRecursive(TSharedPtr<FSequenceTreeItem> InTreeItem);
-
 	/** Expands the input item in the subsequence tree view */
 	void ExpandTreeItem(TSharedPtr<FSequenceTreeItem> InTreeItem) const;
 
