ShotManagement: Added a template level sequence object to the Cine Assembly Schema, which can be built in the schema editor, and used as as starting point when creating assemblies from that schema.
Added a new tab in the schema editor to edit the template sequence in a sequencer timeline.
#rnx

diff --git a/Engine/Plugins/VirtualProduction/CinematicAssemblyTools/Source/CineAssemblyTools/CineAssemblyTools.build.cs b/Engine/Plugins/VirtualProduction/CinematicAssemblyTools/Source/CineAssemblyTools/CineAssemblyTools.build.cs
--- a/Engine/Plugins/VirtualProduction/CinematicAssemblyTools/Source/CineAssemblyTools/CineAssemblyTools.build.cs
+++ b/Engine/Plugins/VirtualProduction/CinematicAssemblyTools/Source/CineAssemblyTools/CineAssemblyTools.build.cs
@@ -27,6 +27,11 @@ namespace UnrealBuildTool.Rules
 					"UniversalObjectLocator"
 				}
 			);
+
+			if (Target.bBuildEditor == true)
+			{
+				PrivateDependencyModuleNames.Add("MovieSceneTools");
+			}
 		}
 	}
 }
diff --git a/Engine/Plugins/VirtualProduction/CinematicAssemblyTools/Source/CineAssemblyTools/Private/CineAssemblySchema.cpp b/Engine/Plugins/VirtualProduction/CinematicAssemblyTools/Source/CineAssemblyTools/Private/CineAssemblySchema.cpp
--- a/Engine/Plugins/VirtualProduction/CinematicAssemblyTools/Source/CineAssemblyTools/Private/CineAssemblySchema.cpp
+++ b/Engine/Plugins/VirtualProduction/CinematicAssemblyTools/Source/CineAssemblyTools/Private/CineAssemblySchema.cpp
@@ -5,9 +5,11 @@
 #include "AssetRegistry/AssetData.h"
 #include "AssetRegistry/AssetRegistryModule.h"
 #include "Modules/ModuleManager.h"
+#include "MovieScene.h"
 
 #if WITH_EDITOR
 #include "AssetToolsModule.h"
+#include "MovieSceneToolsProjectSettings.h"
 #endif
 
 const FName UCineAssemblySchema::SchemaGuidPropertyName = GET_MEMBER_NAME_CHECKED(UCineAssemblySchema, SchemaGuid);
@@ -38,6 +40,19 @@ void UCineAssemblySchema::PostInitProperties()
 	{
 		SchemaGuid = FGuid::NewGuid();
 	}
+
+	if (!HasAnyFlags(RF_ClassDefaultObject | RF_ArchetypeObject) && (!TemplateSequence))
+	{
+		TemplateSequence = NewObject<ULevelSequence>(this);
+		TemplateSequence->Initialize();
+
+#if WITH_EDITOR
+		const UMovieSceneToolsProjectSettings* MovieSceneSettings = GetDefault<UMovieSceneToolsProjectSettings>();
+
+		FFrameRate TickResolution = TemplateSequence->GetMovieScene()->GetTickResolution();
+		TemplateSequence->GetMovieScene()->SetPlaybackRange((MovieSceneSettings->DefaultStartTime * TickResolution).FloorToFrame(), (MovieSceneSettings->DefaultDuration * TickResolution).FloorToFrame().Value);
+#endif // WITH_EDITOR
+	}
 }
 
 void UCineAssemblySchema::PostDuplicate(bool bDuplicateForPIE)
diff --git a/Engine/Plugins/VirtualProduction/CinematicAssemblyTools/Source/CineAssemblyTools/Public/CineAssemblySchema.h b/Engine/Plugins/VirtualProduction/CinematicAssemblyTools/Source/CineAssemblyTools/Public/CineAssemblySchema.h
--- a/Engine/Plugins/VirtualProduction/CinematicAssemblyTools/Source/CineAssemblyTools/Public/CineAssemblySchema.h
+++ b/Engine/Plugins/VirtualProduction/CinematicAssemblyTools/Source/CineAssemblyTools/Public/CineAssemblySchema.h
@@ -3,6 +3,7 @@
 #pragma once
 
 #include "Engine/Texture2D.h"
+#include "LevelSequence.h"
 #include "Styling/SlateBrush.h"
 
 #include "CineAssemblySchema.generated.h"
@@ -118,7 +119,11 @@ public:
 	UPROPERTY()
 	TArray<FString> FoldersToCreate;
 
-	/** Template to use when creating the assembly as a starting point. */
+	/** Template sequence, owned and managed by the schema, to use when creating the assembly as a starting point.  */
+	UPROPERTY()
+	TObjectPtr<ULevelSequence> TemplateSequence;
+
+	/** Path to another level sequence asset to use as a template when creating the assembly as a starting point. */
 	UPROPERTY()
 	FSoftObjectPath Template;
 
diff --git a/Engine/Plugins/VirtualProduction/CinematicAssemblyTools/Source/CineAssemblyToolsEditor/CineAssemblyToolsEditor.build.cs b/Engine/Plugins/VirtualProduction/CinematicAssemblyTools/Source/CineAssemblyToolsEditor/CineAssemblyToolsEditor.build.cs
--- a/Engine/Plugins/VirtualProduction/CinematicAssemblyTools/Source/CineAssemblyToolsEditor/CineAssemblyToolsEditor.build.cs
+++ b/Engine/Plugins/VirtualProduction/CinematicAssemblyTools/Source/CineAssemblyToolsEditor/CineAssemblyToolsEditor.build.cs
@@ -41,6 +41,7 @@ namespace UnrealBuildTool.Rules
 					"NamingTokensUI",
 					"Projects",
 					"PropertyEditor",
+					"Sequencer",
 					"SharedSettingsWidgets",
 					"Slate",
 					"SlateCore",
diff --git a/Engine/Plugins/VirtualProduction/CinematicAssemblyTools/Source/CineAssemblyToolsEditor/Private/CineAssemblyFactory.cpp b/Engine/Plugins/VirtualProduction/CinematicAssemblyTools/Source/CineAssemblyToolsEditor/Private/CineAssemblyFactory.cpp
--- a/Engine/Plugins/VirtualProduction/CinematicAssemblyTools/Source/CineAssemblyToolsEditor/Private/CineAssemblyFactory.cpp
+++ b/Engine/Plugins/VirtualProduction/CinematicAssemblyTools/Source/CineAssemblyToolsEditor/Private/CineAssemblyFactory.cpp
@@ -160,7 +160,13 @@ void UCineAssemblyFactory::CreateConfiguredAssembly(UCineAssembly* ConfiguredAss
 		bool bBuiltFromTemplate = false;
 		if (const UCineAssemblySchema* Schema = ConfiguredAssembly->GetSchema())
 		{
-			if (ULevelSequence* Template = Cast<ULevelSequence>(Schema->Template.TryLoad()))
+			ULevelSequence* Template = Cast<ULevelSequence>(Schema->Template.TryLoad());
+			if (!Template)
+			{
+				Template = Schema->TemplateSequence;
+			}
+
+			if (Template)
 			{
 				ConfiguredAssembly->InitializeFromTemplate(Template);
 				// Confirm we have a valid MovieScene built.
diff --git a/Engine/Plugins/VirtualProduction/CinematicAssemblyTools/Source/CineAssemblyToolsEditor/Private/UI/CineAssembly/SCineAssemblySchemaWindow.cpp b/Engine/Plugins/VirtualProduction/CinematicAssemblyTools/Source/CineAssemblyToolsEditor/Private/UI/CineAssembly/SCineAssemblySchemaWindow.cpp
--- a/Engine/Plugins/VirtualProduction/CinematicAssemblyTools/Source/CineAssemblyToolsEditor/Private/UI/CineAssembly/SCineAssemblySchemaWindow.cpp
+++ b/Engine/Plugins/VirtualProduction/CinematicAssemblyTools/Source/CineAssemblyToolsEditor/Private/UI/CineAssembly/SCineAssemblySchemaWindow.cpp
@@ -9,6 +9,7 @@
 #include "CineAssemblySchemaFactory.h"
 #include "CineAssemblyToolsStyle.h"
 #include "Framework/MultiBox/MultiBoxBuilder.h"
+#include "ISequencerModule.h"
 #include "IStructureDetailsView.h"
 #include "Misc/ConfigCacheIni.h"
 #include "Misc/MessageDialog.h"
@@ -17,6 +18,7 @@
 #include "PropertyCustomizationHelpers.h"
 #include "PropertyEditorDelegates.h"
 #include "PropertyEditorModule.h"
+#include "SequencerSettings.h"
 #include "Settings/ContentBrowserSettings.h"
 #include "Styling/StyleColors.h"
 #include "UObject/Package.h"
@@ -106,6 +108,17 @@ TSharedRef<SWidget> SCineAssemblySchemaWindow::BuildUI()
 	ContentBrowserSettings->SetDisplayEngineFolder(bShowEngineContent);
 	ContentBrowserSettings->SetDisplayPluginFolders(bShowPluginContent);
 
+	// Create the Sequencer that will be displayed on the Timeline tab for editing the template sequence
+	FSequencerInitParams SequencerInitParams;
+	SequencerInitParams.RootSequence = SchemaToConfigure->TemplateSequence;
+	SequencerInitParams.bEditWithinLevelEditor = false;
+	SequencerInitParams.PlaybackContext.Bind(this, &SCineAssemblySchemaWindow::GetPlaybackContextForSequencer);
+
+	SequencerInitParams.ViewParams.UniqueName = TEXT("CineAssemblySchemaTemplateSequencer");
+	SequencerInitParams.HostCapabilities.bSupportsDragAndDrop = false;
+
+	Sequencer = FModuleManager::LoadModuleChecked<ISequencerModule>("Sequencer").CreateSequencer(SequencerInitParams);
+
 	return SNew(SBorder)
 		.BorderImage(FCineAssemblyToolsStyle::Get().GetBrush("ProductionWizard.PanelBackground"))
 		[
@@ -277,6 +290,11 @@ TSharedRef<SWidget> SCineAssemblySchemaWindow::MakeMenuPanel()
 				.Text(LOCTEXT("HierarchyTab", "Content Hierarchy"))
 				.Icon(FSlateIcon(FAppStyle::GetAppStyleSetName(), "Icons.FolderClosed").GetIcon())
 				.HAlign(HAlign_Left)
+
+				+ SSegmentedControl<int32>::Slot(3)
+				.Text(LOCTEXT("TimelineTab", "Timeline Template"))
+				.Icon(FSlateIcon(FCineAssemblyToolsStyle::StyleName, "Icons.Sequencer").GetIcon())
+				.HAlign(HAlign_Left)
 		];
 }
 
@@ -296,6 +314,11 @@ TSharedRef<SWidget> SCineAssemblySchemaWindow::MakeContentPanel()
 		+ SWidgetSwitcher::Slot()
 		[
 			MakeHierarchyTabContent()
+		]
+
+		+ SWidgetSwitcher::Slot()
+		[
+			Sequencer->GetSequencerWidget()
 		];
 
 	return SNew(SVerticalBox)
@@ -1484,6 +1507,21 @@ void SCineAssemblySchemaWindow::UpdateAssetList()
 	}
 }
 
+UObject* SCineAssemblySchemaWindow::GetPlaybackContextForSequencer() const
+{
+	UWorld* EditorWorld = nullptr;
+	for (const FWorldContext& Context : GEngine->GetWorldContexts())
+	{
+		if (Context.WorldType == EWorldType::Editor)
+		{
+			EditorWorld = Context.World();
+			break;
+		}
+	}
+
+	return EditorWorld;
+}
+
 FReply SCineAssemblySchemaWindow::OnAssetRowDragDetected(const FGeometry& MyGeometry, const FPointerEvent& MouseEvent)
 {
 	if (MouseEvent.IsMouseButtonDown(EKeys::LeftMouseButton))
diff --git a/Engine/Plugins/VirtualProduction/CinematicAssemblyTools/Source/CineAssemblyToolsEditor/Private/UI/CineAssembly/SCineAssemblySchemaWindow.h b/Engine/Plugins/VirtualProduction/CinematicAssemblyTools/Source/CineAssemblyToolsEditor/Private/UI/CineAssembly/SCineAssemblySchemaWindow.h
--- a/Engine/Plugins/VirtualProduction/CinematicAssemblyTools/Source/CineAssemblyToolsEditor/Private/UI/CineAssembly/SCineAssemblySchemaWindow.h
+++ b/Engine/Plugins/VirtualProduction/CinematicAssemblyTools/Source/CineAssemblyToolsEditor/Private/UI/CineAssembly/SCineAssemblySchemaWindow.h
@@ -6,6 +6,7 @@
 
 #include "CineAssemblySchema.h"
 #include "Input/DragAndDrop.h"
+#include "ISequencer.h"
 #include "SNamingTokensEditableTextBox.h"
 #include "Widgets/Layout/SWidgetSwitcher.h"
 #include "Widgets/Views/SListView.h"
@@ -272,6 +273,9 @@ private:
 	/** Update the list of assets to create from the current set in the tree view */
 	void UpdateAssetList();
 
+	/** Returns the current editor world to use as the playback context for the timeline template */
+	UObject* GetPlaybackContextForSequencer() const;
+
 	/** Cached content browser settings, used to restore defaults when closing the window */
 	bool bShowEngineContentCached = false;
 	bool bShowPluginContentCached = false;
@@ -312,4 +316,7 @@ private:
 
 	/** The mostly recently added tree item, used to allow the user to immediately rename the item after the tree view updates to show it */
 	TSharedPtr<FSchemaTreeItem> MostRecentlyAddedItem;
+
+	/** The sequencer used to display / edit the template sequence */
+	TSharedPtr<ISequencer> Sequencer;
 };
