ShotManagement: When initializing a Cine Assembly from a template sequence, if there are any folders or tracks whose names contain tokens, those tokens will be resolved and the folder/track will be renamed.
Add warning to schema sequencer when the display rate does not match the active production.
#rnx

diff --git a/Engine/Plugins/VirtualProduction/CinematicAssemblyTools/Source/CineAssemblyTools/Private/CineAssembly.cpp b/Engine/Plugins/VirtualProduction/CinematicAssemblyTools/Source/CineAssemblyTools/Private/CineAssembly.cpp
--- a/Engine/Plugins/VirtualProduction/CinematicAssemblyTools/Source/CineAssemblyTools/Private/CineAssembly.cpp
+++ b/Engine/Plugins/VirtualProduction/CinematicAssemblyTools/Source/CineAssemblyTools/Private/CineAssembly.cpp
@@ -8,6 +8,7 @@
 #include "JsonObjectConverter.h"
 #include "LevelSequenceShotMetaDataLibrary.h"
 #include "MovieScene.h"
+#include "MovieSceneFolder.h"
 #include "NamingTokensEngineSubsystem.h"
 #include "Sections/MovieSceneSubSection.h"
 #include "Serialization/JsonReader.h"
@@ -147,8 +148,13 @@ void UCineAssembly::InitializeFromTemplate(ULevelSequence* Template)
 			return;
 		}
 
+		// Cache the display rate of the original initialized movie scene, which will be restored after the movie scene is replaced with the one from the template
+		const FFrameRate OriginalDisplayRate = MovieScene->GetDisplayRate();
+
 		MovieScene = DuplicateMovieScene;
 
+		MovieScene->SetDisplayRate(OriginalDisplayRate);
+
 		// Fix up supported bindings
 		if (const FMovieSceneBindingReferences* TemplateBindingReferences = Template->GetBindingReferences())
 		{
@@ -165,9 +171,61 @@ void UCineAssembly::InitializeFromTemplate(ULevelSequence* Template)
 				ThisBindingReferencesRef.AddBinding(Reference.ID, MoveTemp(NewLocator), Reference.ResolveFlags, CustomBindingDuplicate);
 			}
 		}
+
+#if WITH_EDITOR
+		// Recursively gather all of the folders in the movie scene
+		TArray<UMovieSceneFolder*> AllFolders;
+		for (UMovieSceneFolder* RootFolder : MovieScene->GetRootFolders())
+		{
+			AllFolders.Add(RootFolder);
+			GetAllMovieSceneFoldersRecursive(RootFolder, AllFolders);
+		}
+
+		// Resolve the tokens in each folder name
+		for (UMovieSceneFolder* Folder : AllFolders)
+		{
+			const FText ResolvedFolderName = UCineAssemblyNamingTokens::GetResolvedText(Folder->GetFolderName().ToString(), this);
+			Folder->SetFolderName(*ResolvedFolderName.ToString());
+		}
+
+		// Resolve the tokens in each track name
+		for (UMovieSceneTrack* Track : MovieScene->GetTracks())
+		{
+			if (UMovieSceneNameableTrack* NameableTrack = Cast<UMovieSceneNameableTrack>(Track))
+			{
+				const FText ResolvedTrackName = UCineAssemblyNamingTokens::GetResolvedText(NameableTrack->GetDisplayName().ToString(), this);
+				NameableTrack->SetDisplayName(ResolvedTrackName);
+			}
+		}
+
+		const TArray<FMovieSceneBinding>& Bindings = Cast<const UMovieScene>(MovieScene)->GetBindings();
+		for (const FMovieSceneBinding& Binding : Bindings)
+		{
+			TArray<UMovieSceneTrack*> Tracks = MovieScene->FindTracks(UMovieSceneNameableTrack::StaticClass(), Binding.GetObjectGuid(), NAME_None);
+			for (UMovieSceneTrack* Track : Tracks)
+			{
+				if (UMovieSceneNameableTrack* NameableTrack = Cast<UMovieSceneNameableTrack>(Track))
+				{
+					const FText ResolvedTrackName = UCineAssemblyNamingTokens::GetResolvedText(NameableTrack->GetDisplayName().ToString(), this);
+					NameableTrack->SetDisplayName(ResolvedTrackName);
+				}
+			}
+		}
+#endif // WITH_EDITOR
 	}
 }
 
+#if WITH_EDITOR
+void UCineAssembly::GetAllMovieSceneFoldersRecursive(UMovieSceneFolder* RootFolder, TArray<UMovieSceneFolder*>& OutFolders)
+{
+	for (UMovieSceneFolder* ChildFolder : RootFolder->GetChildFolders())
+	{
+		OutFolders.Add(ChildFolder);
+		GetAllMovieSceneFoldersRecursive(ChildFolder, OutFolders);
+	}
+}
+#endif // WITH_EDITOR
+
 FGuid UCineAssembly::GetAssemblyGuid() const
 {
 	return AssemblyGuid;
diff --git a/Engine/Plugins/VirtualProduction/CinematicAssemblyTools/Source/CineAssemblyTools/Public/CineAssembly.h b/Engine/Plugins/VirtualProduction/CinematicAssemblyTools/Source/CineAssemblyTools/Public/CineAssembly.h
--- a/Engine/Plugins/VirtualProduction/CinematicAssemblyTools/Source/CineAssemblyTools/Public/CineAssembly.h
+++ b/Engine/Plugins/VirtualProduction/CinematicAssemblyTools/Source/CineAssemblyTools/Public/CineAssembly.h
@@ -10,6 +10,7 @@
 
 #include "CineAssembly.generated.h"
 
+class UMovieSceneFolder;
 class UMovieSceneSubSection;
 
 /** 
@@ -185,6 +186,11 @@ private:
 	/** Get the created FDateTime in local time. If no metadata found, returns an unset optional. */
 	TOptional<FDateTime> TryGetCreatedAsLocalTime() const;
 
+#if WITH_EDITOR
+	/** Recursively gather the movie scene folders within a root folder */
+	void GetAllMovieSceneFoldersRecursive(UMovieSceneFolder* RootFolder, TArray<UMovieSceneFolder*>& OutFolders);
+#endif // WITH_EDITOR
+
 public:
 	/** The assembly name, which supports tokens */
 	UPROPERTY()
diff --git a/Engine/Plugins/VirtualProduction/CinematicAssemblyTools/Source/CineAssemblyToolsEditor/Private/UI/CineAssembly/SCineAssemblySchemaWindow.cpp b/Engine/Plugins/VirtualProduction/CinematicAssemblyTools/Source/CineAssemblyToolsEditor/Private/UI/CineAssembly/SCineAssemblySchemaWindow.cpp
--- a/Engine/Plugins/VirtualProduction/CinematicAssemblyTools/Source/CineAssemblyToolsEditor/Private/UI/CineAssembly/SCineAssemblySchemaWindow.cpp
+++ b/Engine/Plugins/VirtualProduction/CinematicAssemblyTools/Source/CineAssemblyToolsEditor/Private/UI/CineAssembly/SCineAssemblySchemaWindow.cpp
@@ -14,7 +14,9 @@
 #include "Misc/ConfigCacheIni.h"
 #include "Misc/MessageDialog.h"
 #include "Modules/ModuleManager.h"
+#include "MovieScene.h"
 #include "NamingTokensEngineSubsystem.h"
+#include "ProductionSettings.h"
 #include "PropertyCustomizationHelpers.h"
 #include "PropertyEditorDelegates.h"
 #include "PropertyEditorModule.h"
@@ -115,6 +117,10 @@ TSharedRef<SWidget> SCineAssemblySchemaWindow::BuildUI()
 	SequencerInitParams.PlaybackContext.Bind(this, &SCineAssemblySchemaWindow::GetPlaybackContextForSequencer);
 
 	SequencerInitParams.ViewParams.UniqueName = TEXT("CineAssemblySchemaTemplateSequencer");
+
+	SequencerInitParams.ViewParams.ToolbarExtender = MakeShared<FExtender>();
+	SequencerInitParams.ViewParams.ToolbarExtender->AddToolBarExtension("Snapping", EExtensionHook::After, nullptr, FToolBarExtensionDelegate::CreateSP(this, &SCineAssemblySchemaWindow::AddSequencerToolbarExtension));
+
 	SequencerInitParams.HostCapabilities.bSupportsDragAndDrop = false;
 
 	Sequencer = FModuleManager::LoadModuleChecked<ISequencerModule>("Sequencer").CreateSequencer(SequencerInitParams);
@@ -1522,6 +1528,46 @@ UObject* SCineAssemblySchemaWindow::GetPlaybackContextForSequencer() const
 	return EditorWorld;
 }
 
+void SCineAssemblySchemaWindow::AddSequencerToolbarExtension(FToolBarBuilder& Builder)
+{
+	constexpr bool bShouldHaveSeparator = false;
+	Builder.BeginSection("CineAssemblySchemaExtensions", bShouldHaveSeparator);
+
+	TSharedPtr<SWidget> DisplayRateWarningWidget = SNew(SBox)
+		.HAlign(HAlign_Center)
+		.VAlign(VAlign_Center)
+		.ToolTipText(LOCTEXT("DisplayRateWarning", "New Cine Assemblies will be initialized with the Active Production's display rate, not the display rate of this template."))
+		.Visibility(this, &SCineAssemblySchemaWindow::GetDisplayRateWarningVisibility)
+		[
+			SNew(SImage)
+				.Image(FAppStyle::Get().GetBrush("Icons.Warning"))
+				.ColorAndOpacity(FStyleColors::Warning)
+		];
+
+	Builder.AddWidget(DisplayRateWarningWidget.ToSharedRef());
+	Builder.EndSection();
+}
+
+EVisibility SCineAssemblySchemaWindow::GetDisplayRateWarningVisibility() const
+{
+	if (const UProductionSettings* const ProductionSettings = GetDefault<UProductionSettings>())
+	{
+		TOptional<const FCinematicProduction> ActiveProduction = ProductionSettings->GetActiveProduction();
+		if (ActiveProduction.IsSet() && SchemaToConfigure && SchemaToConfigure->TemplateSequence && SchemaToConfigure->TemplateSequence->MovieScene)
+		{
+			const FFrameRate TemplateDisplayRate = SchemaToConfigure->TemplateSequence->MovieScene->GetDisplayRate();
+			const FFrameRate ActiveProductionDisplayRate = ActiveProduction.GetValue().DefaultDisplayRate;
+
+			if (TemplateDisplayRate != ActiveProductionDisplayRate)
+			{
+				return EVisibility::Visible;
+			}
+		}
+	}
+
+	return EVisibility::Collapsed;
+}
+
 FReply SCineAssemblySchemaWindow::OnAssetRowDragDetected(const FGeometry& MyGeometry, const FPointerEvent& MouseEvent)
 {
 	if (MouseEvent.IsMouseButtonDown(EKeys::LeftMouseButton))
diff --git a/Engine/Plugins/VirtualProduction/CinematicAssemblyTools/Source/CineAssemblyToolsEditor/Private/UI/CineAssembly/SCineAssemblySchemaWindow.h b/Engine/Plugins/VirtualProduction/CinematicAssemblyTools/Source/CineAssemblyToolsEditor/Private/UI/CineAssembly/SCineAssemblySchemaWindow.h
--- a/Engine/Plugins/VirtualProduction/CinematicAssemblyTools/Source/CineAssemblyToolsEditor/Private/UI/CineAssembly/SCineAssemblySchemaWindow.h
+++ b/Engine/Plugins/VirtualProduction/CinematicAssemblyTools/Source/CineAssemblyToolsEditor/Private/UI/CineAssembly/SCineAssemblySchemaWindow.h
@@ -276,6 +276,12 @@ private:
 	/** Returns the current editor world to use as the playback context for the timeline template */
 	UObject* GetPlaybackContextForSequencer() const;
 
+	/** Extends the sequencer widget's toolbar */
+	void AddSequencerToolbarExtension(FToolBarBuilder& Builder);
+
+	/** Returns the visibility of the display rate warning toolbar widget */
+	EVisibility GetDisplayRateWarningVisibility() const;
+
 	/** Cached content browser settings, used to restore defaults when closing the window */
 	bool bShowEngineContentCached = false;
 	bool bShowPluginContentCached = false;
