ShotManagement: Add Hidden flag to Cine Assembly Schemas, indicating that they should not appear in the Cine Assembly creation window.
#rnx
#Synced-CL 47206406

diff --git a/Engine/Plugins/VirtualProduction/CinematicAssemblyTools/Source/CineAssemblyTools/Public/CineAssemblySchema.h b/Engine/Plugins/VirtualProduction/CinematicAssemblyTools/Source/CineAssemblyTools/Public/CineAssemblySchema.h
--- a/Engine/Plugins/VirtualProduction/CinematicAssemblyTools/Source/CineAssemblyTools/Public/CineAssemblySchema.h
+++ b/Engine/Plugins/VirtualProduction/CinematicAssemblyTools/Source/CineAssemblyTools/Public/CineAssemblySchema.h
@@ -135,6 +135,10 @@ public:
 	UPROPERTY(AssetRegistrySearchable, BlueprintReadWrite, EditAnywhere, Category = "Default", AdvancedDisplay, meta = (DisplayName = "Data-Only"))
 	bool bIsDataOnly = false;
 
+	/** Hidden schemas will not appear in the UI when creating new Cine Assemblies. */
+	UPROPERTY(AssetRegistrySearchable, BlueprintReadWrite, EditAnywhere, Category = "Default", AdvancedDisplay, meta = (DisplayName = "Hidden"))
+	bool bIsHidden = false;
+
 	static CINEASSEMBLYTOOLS_API const FName SchemaGuidPropertyName;
 
 private:
diff --git a/Engine/Plugins/VirtualProduction/CinematicAssemblyTools/Source/CineAssemblyToolsEditor/Private/UI/CineAssembly/SCineAssemblyConfigWindow.cpp b/Engine/Plugins/VirtualProduction/CinematicAssemblyTools/Source/CineAssemblyToolsEditor/Private/UI/CineAssembly/SCineAssemblyConfigWindow.cpp
--- a/Engine/Plugins/VirtualProduction/CinematicAssemblyTools/Source/CineAssemblyToolsEditor/Private/UI/CineAssembly/SCineAssemblyConfigWindow.cpp
+++ b/Engine/Plugins/VirtualProduction/CinematicAssemblyTools/Source/CineAssemblyToolsEditor/Private/UI/CineAssembly/SCineAssemblyConfigWindow.cpp
@@ -150,6 +150,17 @@ TSharedRef<SWidget> SCineAssemblyConfigWindow::MakeSchemaPanel()
 			OutAssets.Add(NoSchemaAssetData);
 		});
 
+	Config.OnShouldFilterAsset = FOnShouldFilterAsset::CreateLambda([NoSchemaAssetData](const FAssetData& AssetData) -> bool
+		{
+			const FName IsHiddenPropertyName = GET_MEMBER_NAME_CHECKED(UCineAssemblySchema, bIsHidden);
+
+			bool bIsSchemaHidden = false;
+			AssetData.GetTagValue(IsHiddenPropertyName, bIsSchemaHidden);
+
+			// Filter out any schemas that should be hidden
+			return (AssetData == NoSchemaAssetData) ? false : bIsSchemaHidden;
+		});
+
 	Config.InitialAssetSelection = NoSchemaAssetData;
 
 	// The fake NoSchema asset should not display the normal asset tooltip, just a plain text-based tooltip describing what it is
