ShotManagement: Hide the timeline template in schemas that are marked as Data-Only
#rnx
#Synced-CL 47206406

diff --git a/Engine/Plugins/VirtualProduction/CinematicAssemblyTools/Source/CineAssemblyToolsEditor/Private/UI/CineAssembly/SCineAssemblySchemaWindow.cpp b/Engine/Plugins/VirtualProduction/CinematicAssemblyTools/Source/CineAssemblyToolsEditor/Private/UI/CineAssembly/SCineAssemblySchemaWindow.cpp
--- a/Engine/Plugins/VirtualProduction/CinematicAssemblyTools/Source/CineAssemblyToolsEditor/Private/UI/CineAssembly/SCineAssemblySchemaWindow.cpp
+++ b/Engine/Plugins/VirtualProduction/CinematicAssemblyTools/Source/CineAssemblyToolsEditor/Private/UI/CineAssembly/SCineAssemblySchemaWindow.cpp
@@ -301,6 +301,7 @@ TSharedRef<SWidget> SCineAssemblySchemaWindow::MakeMenuPanel()
 				.Text(LOCTEXT("TimelineTab", "Timeline Template"))
 				.Icon(FSlateIcon(FCineAssemblyToolsStyle::StyleName, "Icons.Sequencer").GetIcon())
 				.HAlign(HAlign_Left)
+				.Visibility_Lambda([this]() { return SchemaToConfigure->bIsDataOnly ? EVisibility::Collapsed : EVisibility::Visible; })
 		];
 }
 
