ShotManagement: Add support for adding actors to the Cine Assembly Schema timeline template via drag/drop.
Spawnable bindings can be created, but will not spawn any objects into the current level. This is managed used a custom spawn register.
Possessables can be created (for example by adding an actor from the world outliner) but the bindings will be broken as soon as the actor is added to the sequence, so they do not actually possess anything in the world.
Remove the playback context used by the schema sequencer to prevent most tracks from affecting the current level.
Fix the schema window to manually close the sequencer before the window is closed.

#rnx
#Synced-CL 48875823

diff --git a/Engine/Plugins/VirtualProduction/CinematicAssemblyTools/Source/CineAssemblyToolsEditor/Private/CineAssemblySchemaFactory.cpp b/Engine/Plugins/VirtualProduction/CinematicAssemblyTools/Source/CineAssemblyToolsEditor/Private/CineAssemblySchemaFactory.cpp
--- a/Engine/Plugins/VirtualProduction/CinematicAssemblyTools/Source/CineAssemblyToolsEditor/Private/CineAssemblySchemaFactory.cpp
+++ b/Engine/Plugins/VirtualProduction/CinematicAssemblyTools/Source/CineAssemblyToolsEditor/Private/CineAssemblySchemaFactory.cpp
@@ -43,6 +43,7 @@ bool UCineAssemblySchemaFactory::ConfigureProperties()
 		.SupportsMinimize(false);
 
 	NewSchemaWindow->SetContent(SchemaWidget);
+	NewSchemaWindow->SetOnWindowClosed(FOnWindowClosed::CreateSP(SchemaWidget, &SCineAssemblySchemaWindow::OnWindowClosed));
 
 	const IMainFrameModule& MainFrameModule = FModuleManager::LoadModuleChecked<IMainFrameModule>(TEXT("MainFrame"));
 	const TSharedPtr<SWindow> ParentWindow = MainFrameModule.GetParentWindow();
diff --git a/Engine/Plugins/VirtualProduction/CinematicAssemblyTools/Source/CineAssemblyToolsEditor/Private/Sequencer/CineAssemblySchemaSpawnRegister.cpp b/Engine/Plugins/VirtualProduction/CinematicAssemblyTools/Source/CineAssemblyToolsEditor/Private/Sequencer/CineAssemblySchemaSpawnRegister.cpp
new file mode 100644
--- /dev/null
+++ b/Engine/Plugins/VirtualProduction/CinematicAssemblyTools/Source/CineAssemblyToolsEditor/Private/Sequencer/CineAssemblySchemaSpawnRegister.cpp
@@ -0,0 +1,57 @@
+// Copyright Epic Games, Inc. All Rights Reserved.
+
+#include "CineAssemblySchemaSpawnRegister.h"
+#include "ILevelSequenceModule.h"
+#include "Modules/ModuleManager.h"
+
+FCineAssemblySchemaSpawnRegister::FCineAssemblySchemaSpawnRegister()
+{
+	// Generate a set of object spawners that will be used to set up the defaults for supported spawnable types (primarily actors)
+	ILevelSequenceModule& LevelSequenceModule = FModuleManager::GetModuleChecked<ILevelSequenceModule>("LevelSequence");
+	LevelSequenceModule.GenerateObjectSpawners(MovieSceneObjectSpawners);
+}
+
+UObject* FCineAssemblySchemaSpawnRegister::SpawnObject(const FGuid& BindingId, UMovieScene& MovieScene, FMovieSceneSequenceIDRef Template, TSharedRef<const FSharedPlaybackState> SharedPlaybackState, int32 BindingIndex)
+{
+	// Spawnables in a schema template sequence never actually get spawned in the world
+	return nullptr;
+}
+
+UObject* FCineAssemblySchemaSpawnRegister::SpawnObject(FMovieSceneSpawnable& Spawnable, FMovieSceneSequenceIDRef TemplateID, TSharedRef<const FSharedPlaybackState> SharedPlaybackState)
+{
+	// Spawnables in a schema template sequence never actually get spawned in the world
+	return nullptr;
+}
+
+void FCineAssemblySchemaSpawnRegister::DestroySpawnedObject(UObject& Object, UMovieSceneSpawnableBindingBase* CustomSpawnableBinding)
+{
+	// Nothing was ever spawned, so there is nothing to destroy
+}
+
+#if WITH_EDITOR
+
+bool FCineAssemblySchemaSpawnRegister::CanSpawnObject(UClass* InClass) const
+{
+	for (TSharedRef<IMovieSceneObjectSpawner> MovieSceneObjectSpawner : MovieSceneObjectSpawners)
+	{
+		if (InClass->IsChildOf(MovieSceneObjectSpawner->GetSupportedTemplateType()))
+		{
+			return true;
+		}
+	}
+	return false;
+}
+
+void FCineAssemblySchemaSpawnRegister::SetupDefaultsForSpawnable(UObject* SpawnedObject, const FGuid& Guid, const TOptional<FTransformData>& TransformData, TSharedRef<ISequencer> Sequencer, USequencerSettings* Settings)
+{
+	for (TSharedPtr<IMovieSceneObjectSpawner> MovieSceneObjectSpawner : MovieSceneObjectSpawners)
+	{
+		if (MovieSceneObjectSpawner->CanSetupDefaultsForSpawnable(SpawnedObject))
+		{
+			MovieSceneObjectSpawner->SetupDefaultsForSpawnable(SpawnedObject, Guid, TransformData, Sequencer, Settings);
+			return;
+		}
+	}
+}
+
+#endif // WITH_EDITOR
diff --git a/Engine/Plugins/VirtualProduction/CinematicAssemblyTools/Source/CineAssemblyToolsEditor/Private/Sequencer/CineAssemblySchemaSpawnRegister.h b/Engine/Plugins/VirtualProduction/CinematicAssemblyTools/Source/CineAssemblyToolsEditor/Private/Sequencer/CineAssemblySchemaSpawnRegister.h
new file mode 100644
--- /dev/null
+++ b/Engine/Plugins/VirtualProduction/CinematicAssemblyTools/Source/CineAssemblyToolsEditor/Private/Sequencer/CineAssemblySchemaSpawnRegister.h
@@ -0,0 +1,29 @@
+// Copyright Epic Games, Inc. All Rights Reserved.
+
+#pragma once
+
+#include "MovieSceneSpawnRegister.h"
+
+/** Movie scene spawn register that handles spawning objects in a Cine Assembly Schema's template sequence */
+class FCineAssemblySchemaSpawnRegister : public FMovieSceneSpawnRegister
+{
+public:
+	FCineAssemblySchemaSpawnRegister();
+	~FCineAssemblySchemaSpawnRegister() = default;
+
+protected:
+	//~ Begin FMovieSceneSpawnRegister interface
+	virtual UObject* SpawnObject(const FGuid& BindingId, UMovieScene& MovieScene, FMovieSceneSequenceIDRef Template, TSharedRef<const FSharedPlaybackState> SharedPlaybackState, int32 BindingIndex) override;
+	virtual UObject* SpawnObject(FMovieSceneSpawnable& Spawnable, FMovieSceneSequenceIDRef TemplateID, TSharedRef<const FSharedPlaybackState> SharedPlaybackState) override;
+
+	virtual void DestroySpawnedObject(UObject& Object, UMovieSceneSpawnableBindingBase* CustomSpawnableBinding) override;
+
+#if WITH_EDITOR
+	virtual bool CanSpawnObject(UClass* InClass) const override;
+	virtual void SetupDefaultsForSpawnable(UObject* SpawnedObject, const FGuid& Guid, const TOptional<FTransformData>& TransformData, TSharedRef<ISequencer> Sequencer, USequencerSettings* Settings) override;
+#endif
+	//~ End FMovieSceneSpawnRegister interface
+
+private:
+	TArray<TSharedRef<IMovieSceneObjectSpawner>> MovieSceneObjectSpawners;
+};
diff --git a/Engine/Plugins/VirtualProduction/CinematicAssemblyTools/Source/CineAssemblyToolsEditor/Private/UI/CineAssembly/SCineAssemblySchemaWindow.cpp b/Engine/Plugins/VirtualProduction/CinematicAssemblyTools/Source/CineAssemblyToolsEditor/Private/UI/CineAssembly/SCineAssemblySchemaWindow.cpp
--- a/Engine/Plugins/VirtualProduction/CinematicAssemblyTools/Source/CineAssemblyToolsEditor/Private/UI/CineAssembly/SCineAssemblySchemaWindow.cpp
+++ b/Engine/Plugins/VirtualProduction/CinematicAssemblyTools/Source/CineAssemblyToolsEditor/Private/UI/CineAssembly/SCineAssemblySchemaWindow.cpp
@@ -20,6 +20,8 @@
 #include "PropertyCustomizationHelpers.h"
 #include "PropertyEditorDelegates.h"
 #include "PropertyEditorModule.h"
+#include "Sequencer/CineAssemblySchemaSpawnRegister.h"
+#include "SequencerUtilities.h"
 #include "SequencerSettings.h"
 #include "Settings/ContentBrowserSettings.h"
 #include "Styling/StyleColors.h"
@@ -111,19 +113,25 @@ TSharedRef<SWidget> SCineAssemblySchemaWindow::BuildUI()
 	ContentBrowserSettings->SetDisplayPluginFolders(bShowPluginContent);
 
 	// Create the Sequencer that will be displayed on the Timeline tab for editing the template sequence
+	TSharedRef<FCineAssemblySchemaSpawnRegister> SpawnRegister = MakeShared<FCineAssemblySchemaSpawnRegister>();
+
 	FSequencerInitParams SequencerInitParams;
 	SequencerInitParams.RootSequence = SchemaToConfigure->TemplateSequence;
 	SequencerInitParams.bEditWithinLevelEditor = false;
-	SequencerInitParams.PlaybackContext.Bind(this, &SCineAssemblySchemaWindow::GetPlaybackContextForSequencer);
+	SequencerInitParams.SpawnRegister = SpawnRegister;
 
 	SequencerInitParams.ViewParams.UniqueName = TEXT("CineAssemblySchemaTemplateSequencer");
 
 	SequencerInitParams.ViewParams.ToolbarExtender = MakeShared<FExtender>();
 	SequencerInitParams.ViewParams.ToolbarExtender->AddToolBarExtension("Snapping", EExtensionHook::After, nullptr, FToolBarExtensionDelegate::CreateSP(this, &SCineAssemblySchemaWindow::AddSequencerToolbarExtension));
 
-	SequencerInitParams.HostCapabilities.bSupportsDragAndDrop = false;
+	//SequencerInitParams.HostCapabilities.bSupportsCurveEditor = true;
+	SequencerInitParams.HostCapabilities.bSupportsAddFromContentBrowser = true;
+	SequencerInitParams.HostCapabilities.bSupportsDragAndDrop = true;
+	SequencerInitParams.HostCapabilities.bSupportsSidebar = true;
 
 	Sequencer = FModuleManager::LoadModuleChecked<ISequencerModule>("Sequencer").CreateSequencer(SequencerInitParams);
+	Sequencer->OnActorAddedToSequencer().AddSP(this, &SCineAssemblySchemaWindow::OnActorAddedToSequencer);
 
 	return SNew(SBorder)
 		.BorderImage(FCineAssemblyToolsStyle::Get().GetBrush("ProductionWizard.PanelBackground"))
@@ -221,6 +229,13 @@ FString SCineAssemblySchemaWindow::GetSchemaName()
 	return TEXT("CineAssemblySchema");
 }
 
+void SCineAssemblySchemaWindow::OnWindowClosed(const TSharedRef<SWindow>& InWindow)
+{
+	if (Sequencer)
+	{
+		Sequencer->Close();
+	}
+}
 
 bool SCineAssemblySchemaWindow::DoesSchemaExistWithName(const FString& SchemaName) const
 {
@@ -1514,19 +1529,11 @@ void SCineAssemblySchemaWindow::UpdateAssetList()
 	}
 }
 
-UObject* SCineAssemblySchemaWindow::GetPlaybackContextForSequencer() const
+void SCineAssemblySchemaWindow::OnActorAddedToSequencer(AActor* InActor, const FGuid InBinding)
 {
-	UWorld* EditorWorld = nullptr;
-	for (const FWorldContext& Context : GEngine->GetWorldContexts())
-	{
-		if (Context.WorldType == EWorldType::Editor)
-		{
-			EditorWorld = Context.World();
-			break;
-		}
-	}
-
-	return EditorWorld;
+	// Possessables in our template sequence may not actually bind to any actors in the world
+	FMovieSceneBindingProxy Proxy = FMovieSceneBindingProxy(InBinding, SchemaToConfigure->TemplateSequence);
+	FSequencerUtilities::RemoveActorsFromBinding(Sequencer.ToSharedRef(), { InActor }, Proxy);
 }
 
 void SCineAssemblySchemaWindow::AddSequencerToolbarExtension(FToolBarBuilder& Builder)
diff --git a/Engine/Plugins/VirtualProduction/CinematicAssemblyTools/Source/CineAssemblyToolsEditor/Private/UI/CineAssembly/SCineAssemblySchemaWindow.h b/Engine/Plugins/VirtualProduction/CinematicAssemblyTools/Source/CineAssemblyToolsEditor/Private/UI/CineAssembly/SCineAssemblySchemaWindow.h
--- a/Engine/Plugins/VirtualProduction/CinematicAssemblyTools/Source/CineAssemblyToolsEditor/Private/UI/CineAssembly/SCineAssemblySchemaWindow.h
+++ b/Engine/Plugins/VirtualProduction/CinematicAssemblyTools/Source/CineAssemblyToolsEditor/Private/UI/CineAssembly/SCineAssemblySchemaWindow.h
@@ -134,6 +134,9 @@ public:
 	/** Searches the asset registry for a Cine Assembly Schema matching the input ID and updates the UI */
 	void FindSchema(FGuid SchemaID);
 
+	/** Callback to cleanup this widget when the containing window is closed */
+	void OnWindowClosed(const TSharedRef<SWindow>& InWindow);
+
 private:
 	/** Constructs the main UI for the widget */
 	TSharedRef<SWidget> BuildUI();
@@ -273,8 +276,8 @@ private:
 	/** Update the list of assets to create from the current set in the tree view */
 	void UpdateAssetList();
 
-	/** Returns the current editor world to use as the playback context for the timeline template */
-	UObject* GetPlaybackContextForSequencer() const;
+	/** Callback when an actor is added to the template sequence, used to remove the actor(s) from the binding */
+	void OnActorAddedToSequencer(AActor* InActor, const FGuid InBinding);
 
 	/** Extends the sequencer widget's toolbar */
 	void AddSequencerToolbarExtension(FToolBarBuilder& Builder);
