ShotManagement: Fix crash when adding a media plate actor to the schema sequencer.
The Schema window binds to the Sequencer's OnActorAddedToSequencer event in order to remove actor bindings, which is required to make sure possessables don't actually reference any actors in the world.
However, other things, like the MediaPlate Track Editor also bind to that event, and need the binding to still be valid before we break the actor association. And we have no control over the order of execution of multicast delegates.
This CL changes the behavior of the Schema window to delay breaking the actor binding until the next Tick, effectively making sure it executes last.
#rnx

#Synced-CL 49044222

diff --git a/Engine/Plugins/VirtualProduction/CinematicAssemblyTools/Source/CineAssemblyToolsEditor/Private/UI/CineAssembly/SCineAssemblySchemaWindow.cpp b/Engine/Plugins/VirtualProduction/CinematicAssemblyTools/Source/CineAssemblyToolsEditor/Private/UI/CineAssembly/SCineAssemblySchemaWindow.cpp
--- a/Engine/Plugins/VirtualProduction/CinematicAssemblyTools/Source/CineAssemblyToolsEditor/Private/UI/CineAssembly/SCineAssemblySchemaWindow.cpp
+++ b/Engine/Plugins/VirtualProduction/CinematicAssemblyTools/Source/CineAssemblyToolsEditor/Private/UI/CineAssembly/SCineAssemblySchemaWindow.cpp
@@ -8,6 +8,7 @@
 #include "CineAssemblyNamingTokens.h"
 #include "CineAssemblySchemaFactory.h"
 #include "CineAssemblyToolsStyle.h"
+#include "Editor.h"
 #include "Framework/MultiBox/MultiBoxBuilder.h"
 #include "ISequencerModule.h"
 #include "IStructureDetailsView.h"
@@ -1532,8 +1533,18 @@ void SCineAssemblySchemaWindow::UpdateAssetList()
 void SCineAssemblySchemaWindow::OnActorAddedToSequencer(AActor* InActor, const FGuid InBinding)
 {
 	// Possessables in our template sequence may not actually bind to any actors in the world
-	FMovieSceneBindingProxy Proxy = FMovieSceneBindingProxy(InBinding, SchemaToConfigure->TemplateSequence);
-	FSequencerUtilities::RemoveActorsFromBinding(Sequencer.ToSharedRef(), { InActor }, Proxy);
+	if (InActor && InBinding.IsValid())
+	{
+		FMovieSceneBindingProxy Proxy = FMovieSceneBindingProxy(InBinding, SchemaToConfigure->TemplateSequence);
+
+		// There are likely other delegates that are bound to OnActorAddedToSequencer, but we cannot control the order in which they execute.
+		// It might be dangerous to remove the actor from the newly created binding too early, before other delegates have a chance to do something with that binding.
+		// Therefore, we delay our action until the next tick to ensure that this do execute AFTER any other bound delegates. 
+		GEditor->GetTimerManager()->SetTimerForNextTick(FTimerDelegate::CreateLambda([this, Proxy]()
+			{
+				FSequencerUtilities::RemoveAllBindings(Sequencer.ToSharedRef(), Proxy);
+			}));
+	}
 }
 
 void SCineAssemblySchemaWindow::AddSequencerToolbarExtension(FToolBarBuilder& Builder)
