ShotManagement: Change the template sequence that is created by [REDACTED_PERSON] to be a UCineAssembly, rather than a ULevelSequence.
This will give us additional flexibility when customizing the template sequence, such as marking certain track types as being unsupported in the Schema Sequencer.
#rnx
#Synced-CL 49044222

diff --git a/Engine/Plugins/VirtualProduction/CinematicAssemblyTools/Source/CineAssemblyTools/Private/CineAssemblySchema.cpp b/Engine/Plugins/VirtualProduction/CinematicAssemblyTools/Source/CineAssemblyTools/Private/CineAssemblySchema.cpp
--- a/Engine/Plugins/VirtualProduction/CinematicAssemblyTools/Source/CineAssemblyTools/Private/CineAssemblySchema.cpp
+++ b/Engine/Plugins/VirtualProduction/CinematicAssemblyTools/Source/CineAssemblyTools/Private/CineAssemblySchema.cpp
@@ -4,6 +4,7 @@
 
 #include "AssetRegistry/AssetData.h"
 #include "AssetRegistry/AssetRegistryModule.h"
+#include "CineAssembly.h"
 #include "Modules/ModuleManager.h"
 #include "MovieScene.h"
 
@@ -43,7 +44,7 @@ void UCineAssemblySchema::PostInitProperties()
 
 	if (!HasAnyFlags(RF_ClassDefaultObject | RF_ArchetypeObject) && (!TemplateSequence))
 	{
-		TemplateSequence = NewObject<ULevelSequence>(this);
+		TemplateSequence = NewObject<UCineAssembly>(this);
 		TemplateSequence->Initialize();
 
 #if WITH_EDITOR
diff --git a/Engine/Plugins/VirtualProduction/CinematicAssemblyTools/Source/CineAssemblyTools/Public/CineAssemblySchema.h b/Engine/Plugins/VirtualProduction/CinematicAssemblyTools/Source/CineAssemblyTools/Public/CineAssemblySchema.h
--- a/Engine/Plugins/VirtualProduction/CinematicAssemblyTools/Source/CineAssemblyTools/Public/CineAssemblySchema.h
+++ b/Engine/Plugins/VirtualProduction/CinematicAssemblyTools/Source/CineAssemblyTools/Public/CineAssemblySchema.h
@@ -8,6 +8,8 @@
 
 #include "CineAssemblySchema.generated.h"
 
+class UCineAssembly;
+
 /** The types of assembly metadata supported by Cine Assembly Schemas */
 UENUM()
 enum class ECineAssemblyMetadataType : uint8
@@ -121,7 +123,7 @@ public:
 
 	/** Template sequence, owned and managed by the schema, to use when creating the assembly as a starting point.  */
 	UPROPERTY()
-	TObjectPtr<ULevelSequence> TemplateSequence;
+	TObjectPtr<UCineAssembly> TemplateSequence;
 
 	/** Path to another level sequence asset to use as a template when creating the assembly as a starting point. */
 	UPROPERTY()
