ShotManagement: Changes the UCineAssembly to no longer support Event tracks ONLY when the object that owns it is a UCineAssemblySchema (i.e. only apply this behavior to template sequences that can be edited in the schema window.)
This effectively hides the Event Track as an option from the Add Track in the Schema's Sequencer, while still allowing Event Tracks in all other Assemblies.
#rnx
#Synced-CL 49044222

diff --git a/Engine/Plugins/VirtualProduction/CinematicAssemblyTools/Source/CineAssemblyTools/CineAssemblyTools.build.cs b/Engine/Plugins/VirtualProduction/CinematicAssemblyTools/Source/CineAssemblyTools/CineAssemblyTools.build.cs
--- a/Engine/Plugins/VirtualProduction/CinematicAssemblyTools/Source/CineAssemblyTools/CineAssemblyTools.build.cs
+++ b/Engine/Plugins/VirtualProduction/CinematicAssemblyTools/Source/CineAssemblyTools/CineAssemblyTools.build.cs
@@ -23,6 +23,7 @@ namespace UnrealBuildTool.Rules
 					"Json",
 					"JsonUtilities",
 					"MovieScene",
+					"MovieSceneTracks",
 					"SlateCore",
 					"UniversalObjectLocator"
 				}
diff --git a/Engine/Plugins/VirtualProduction/CinematicAssemblyTools/Source/CineAssemblyTools/Private/CineAssembly.cpp b/Engine/Plugins/VirtualProduction/CinematicAssemblyTools/Source/CineAssemblyTools/Private/CineAssembly.cpp
--- a/Engine/Plugins/VirtualProduction/CinematicAssemblyTools/Source/CineAssemblyTools/Private/CineAssembly.cpp
+++ b/Engine/Plugins/VirtualProduction/CinematicAssemblyTools/Source/CineAssemblyTools/Private/CineAssembly.cpp
@@ -14,6 +14,7 @@
 #include "Serialization/JsonReader.h"
 #include "Serialization/JsonSerializer.h"
 #include "Serialization/JsonWriter.h"
+#include "Tracks/MovieSceneEventTrack.h"
 #include "Tracks/MovieSceneSubTrack.h"
 #include "UniversalObjectLocator.h"
 #include "UObject/AssetRegistryTagsContext.h"
@@ -123,6 +124,22 @@ bool UCineAssembly::IsCompatibleAsSubSequence(const UMovieSceneSequence& ParentS
 	return !bIsDataOnly;
 }
 
+#if WITH_EDITOR
+ETrackSupport UCineAssembly::IsTrackSupportedImpl(TSubclassOf<class UMovieSceneTrack> InTrackClass) const
+{
+	if (GetTypedOuter<UCineAssemblySchema>())
+	{
+		if (InTrackClass == UMovieSceneEventTrack::StaticClass())
+		{
+			return ETrackSupport::NotSupported;
+		}
+	}
+
+	return Super::IsTrackSupportedImpl(InTrackClass);
+}
+#endif //WITH_EDITOR
+
+
 void UCineAssembly::Initialize()
 {
 	Super::Initialize();
diff --git a/Engine/Plugins/VirtualProduction/CinematicAssemblyTools/Source/CineAssemblyTools/Public/CineAssembly.h b/Engine/Plugins/VirtualProduction/CinematicAssemblyTools/Source/CineAssemblyTools/Public/CineAssembly.h
--- a/Engine/Plugins/VirtualProduction/CinematicAssemblyTools/Source/CineAssemblyTools/Public/CineAssembly.h
+++ b/Engine/Plugins/VirtualProduction/CinematicAssemblyTools/Source/CineAssemblyTools/Public/CineAssembly.h
@@ -31,6 +31,10 @@ public:
 	//~ End UMovieSceneSequence Interface
 
 	//~ Begin ULevelSquence Interface
+#if WITH_EDITOR
+	CINEASSEMBLYTOOLS_API virtual ETrackSupport IsTrackSupportedImpl(TSubclassOf<class UMovieSceneTrack> InTrackClass) const override;
+#endif //WITH_EDITOR
+
 	CINEASSEMBLYTOOLS_API virtual void Initialize() override;
 	//~ End ULevelSquence Interface
 
