ShotManagement: Force cvar tracks in the Schema sequencer to be muted so they cannot evaluate.
#rnx
#Synced-CL 49044222

diff --git a/Engine/Plugins/VirtualProduction/CinematicAssemblyTools/Source/CineAssemblyToolsEditor/CineAssemblyToolsEditor.build.cs b/Engine/Plugins/VirtualProduction/CinematicAssemblyTools/Source/CineAssemblyToolsEditor/CineAssemblyToolsEditor.build.cs
--- a/Engine/Plugins/VirtualProduction/CinematicAssemblyTools/Source/CineAssemblyToolsEditor/CineAssemblyToolsEditor.build.cs
+++ b/Engine/Plugins/VirtualProduction/CinematicAssemblyTools/Source/CineAssemblyToolsEditor/CineAssemblyToolsEditor.build.cs
@@ -37,6 +37,7 @@ namespace UnrealBuildTool.Rules
 					"MovieRenderPipelineCore",
 					"MovieScene",
 					"MovieSceneTools",
+					"MovieSceneTracks",
 					"NamingTokens",
 					"NamingTokensUI",
 					"Projects",
diff --git a/Engine/Plugins/VirtualProduction/CinematicAssemblyTools/Source/CineAssemblyToolsEditor/Private/UI/CineAssembly/SCineAssemblySchemaWindow.cpp b/Engine/Plugins/VirtualProduction/CinematicAssemblyTools/Source/CineAssemblyToolsEditor/Private/UI/CineAssembly/SCineAssemblySchemaWindow.cpp
--- a/Engine/Plugins/VirtualProduction/CinematicAssemblyTools/Source/CineAssemblyToolsEditor/Private/UI/CineAssembly/SCineAssemblySchemaWindow.cpp
+++ b/Engine/Plugins/VirtualProduction/CinematicAssemblyTools/Source/CineAssemblyToolsEditor/Private/UI/CineAssembly/SCineAssemblySchemaWindow.cpp
@@ -26,6 +26,7 @@
 #include "SequencerSettings.h"
 #include "Settings/ContentBrowserSettings.h"
 #include "Styling/StyleColors.h"
+#include "Tracks/MovieSceneCVarTrack.h"
 #include "UObject/Package.h"
 #include "Widgets/Input/SButton.h"
 #include "Widgets/Input/SEditableTextBox.h"
@@ -133,6 +134,7 @@ TSharedRef<SWidget> SCineAssemblySchemaWindow::BuildUI()
 
 	Sequencer = FModuleManager::LoadModuleChecked<ISequencerModule>("Sequencer").CreateSequencer(SequencerInitParams);
 	Sequencer->OnActorAddedToSequencer().AddSP(this, &SCineAssemblySchemaWindow::OnActorAddedToSequencer);
+	Sequencer->OnMovieSceneDataChanged().AddSP(this, &SCineAssemblySchemaWindow::OnMovieSceneDataChanged);
 
 	return SNew(SBorder)
 		.BorderImage(FCineAssemblyToolsStyle::Get().GetBrush("ProductionWizard.PanelBackground"))
@@ -1530,6 +1532,27 @@ void SCineAssemblySchemaWindow::UpdateAssetList()
 	}
 }
 
+void SCineAssemblySchemaWindow::OnMovieSceneDataChanged(EMovieSceneDataChangeType InChangeType)
+{
+	// In order to mute cvar tracks, we only care about changes when a track is added (in case it was a cvar track)
+	// or when a track value changes (which includes the user trying to unmute the track)
+	if ((InChangeType == EMovieSceneDataChangeType::MovieSceneStructureItemAdded) || (InChangeType == EMovieSceneDataChangeType::TrackValueChanged))
+	{
+		if (UMovieScene* MovieScene = SchemaToConfigure->TemplateSequence->GetMovieScene())
+		{
+			// Because we do not have any additional information about what the actual MovieScene change was, we just loop through every track
+			// and force all of the cvar tracks to be muted
+			for (UMovieSceneTrack* Track : MovieScene->GetTracks())
+			{
+				if (UMovieSceneCVarTrack* CVarTrack = Cast<UMovieSceneCVarTrack>(Track))
+				{
+					CVarTrack->SetLocalEvalDisabled(true);
+				}
+			}
+		}
+	}
+}
+
 void SCineAssemblySchemaWindow::OnActorAddedToSequencer(AActor* InActor, const FGuid InBinding)
 {
 	// Possessables in our template sequence may not actually bind to any actors in the world
diff --git a/Engine/Plugins/VirtualProduction/CinematicAssemblyTools/Source/CineAssemblyToolsEditor/Private/UI/CineAssembly/SCineAssemblySchemaWindow.h b/Engine/Plugins/VirtualProduction/CinematicAssemblyTools/Source/CineAssemblyToolsEditor/Private/UI/CineAssembly/SCineAssemblySchemaWindow.h
--- a/Engine/Plugins/VirtualProduction/CinematicAssemblyTools/Source/CineAssemblyToolsEditor/Private/UI/CineAssembly/SCineAssemblySchemaWindow.h
+++ b/Engine/Plugins/VirtualProduction/CinematicAssemblyTools/Source/CineAssemblyToolsEditor/Private/UI/CineAssembly/SCineAssemblySchemaWindow.h
@@ -279,6 +279,9 @@ private:
 	/** Callback when an actor is added to the template sequence, used to remove the actor(s) from the binding */
 	void OnActorAddedToSequencer(AActor* InActor, const FGuid InBinding);
 
+	/** Callback whenever MovieScene data in the template sequence is changed */
+	void OnMovieSceneDataChanged(EMovieSceneDataChangeType InChangeType);
+
 	/** Extends the sequencer widget's toolbar */
 	void AddSequencerToolbarExtension(FToolBarBuilder& Builder);
 
