ShotManagement: Add the "Add CineCamera" button to the toolbar in the Schema Sequencer.
#rnx
#Synced-CL 49044222

diff --git a/Engine/Plugins/VirtualProduction/CinematicAssemblyTools/Source/CineAssemblyToolsEditor/CineAssemblyToolsEditor.build.cs b/Engine/Plugins/VirtualProduction/CinematicAssemblyTools/Source/CineAssemblyToolsEditor/CineAssemblyToolsEditor.build.cs
--- a/Engine/Plugins/VirtualProduction/CinematicAssemblyTools/Source/CineAssemblyToolsEditor/CineAssemblyToolsEditor.build.cs
+++ b/Engine/Plugins/VirtualProduction/CinematicAssemblyTools/Source/CineAssemblyToolsEditor/CineAssemblyToolsEditor.build.cs
@@ -19,6 +19,7 @@ namespace UnrealBuildTool.Rules
 					"AssetRegistry",
 					"AssetTools",
 					"CineAssemblyTools",
+					"CinematicCamera",
 					"ClassViewer",
 					"ContentBrowser",
 					"ContentBrowserData",
diff --git a/Engine/Plugins/VirtualProduction/CinematicAssemblyTools/Source/CineAssemblyToolsEditor/Private/Sequencer/CineAssemblySchemaSequencerUtilities.cpp b/Engine/Plugins/VirtualProduction/CinematicAssemblyTools/Source/CineAssemblyToolsEditor/Private/Sequencer/CineAssemblySchemaSequencerUtilities.cpp
new file mode 100644
--- /dev/null
+++ b/Engine/Plugins/VirtualProduction/CinematicAssemblyTools/Source/CineAssemblyToolsEditor/Private/Sequencer/CineAssemblySchemaSequencerUtilities.cpp
@@ -0,0 +1,22 @@
+// Copyright Epic Games, Inc. All Rights Reserved.
+
+#include "CineAssemblySchemaSequencerUtilities.h"
+
+#include "CineCameraActor.h"
+#include "ISequencer.h"
+#include "SequencerUtilities.h"
+
+void FCineAssemblySchemaSequencerUtilities::CreateCamera(TSharedRef<ISequencer> Sequencer)
+{
+	// Create the new cine camera actor, but do not spawn it in the world
+	ACineCameraActor* NewCamera = NewObject<ACineCameraActor>();
+
+	// Add a new spawnable binding to the sequence for the new cine camera
+	UE::Sequencer::FCreateBindingParams Params;
+	Params.bAllowCustomBinding = true;
+	Params.bSpawnable = true;
+	Params.bSetupDefaults = true;
+	Params.BindingNameOverride = TEXT("CineCameraActor");
+
+	FSequencerUtilities::CreateBinding(Sequencer, *NewCamera, Params);
+}
diff --git a/Engine/Plugins/VirtualProduction/CinematicAssemblyTools/Source/CineAssemblyToolsEditor/Private/Sequencer/CineAssemblySchemaSequencerUtilities.h b/Engine/Plugins/VirtualProduction/CinematicAssemblyTools/Source/CineAssemblyToolsEditor/Private/Sequencer/CineAssemblySchemaSequencerUtilities.h
new file mode 100644
--- /dev/null
+++ b/Engine/Plugins/VirtualProduction/CinematicAssemblyTools/Source/CineAssemblyToolsEditor/Private/Sequencer/CineAssemblySchemaSequencerUtilities.h
@@ -0,0 +1,14 @@
+// Copyright Epic Games, Inc. All Rights Reserved.
+
+#pragma once
+
+#include "CoreMinimal.h"
+
+class ISequencer;
+
+/** Utility functions for the custom Sequencer used by the Cine Assembly Schema */
+struct FCineAssemblySchemaSequencerUtilities
+{
+public:
+	static void CreateCamera(TSharedRef<ISequencer> Sequencer);
+};
diff --git a/Engine/Plugins/VirtualProduction/CinematicAssemblyTools/Source/CineAssemblyToolsEditor/Private/UI/CineAssembly/SCineAssemblySchemaWindow.cpp b/Engine/Plugins/VirtualProduction/CinematicAssemblyTools/Source/CineAssemblyToolsEditor/Private/UI/CineAssembly/SCineAssemblySchemaWindow.cpp
--- a/Engine/Plugins/VirtualProduction/CinematicAssemblyTools/Source/CineAssemblyToolsEditor/Private/UI/CineAssembly/SCineAssemblySchemaWindow.cpp
+++ b/Engine/Plugins/VirtualProduction/CinematicAssemblyTools/Source/CineAssemblyToolsEditor/Private/UI/CineAssembly/SCineAssemblySchemaWindow.cpp
@@ -21,6 +21,7 @@
 #include "PropertyCustomizationHelpers.h"
 #include "PropertyEditorDelegates.h"
 #include "PropertyEditorModule.h"
+#include "Sequencer/CineAssemblySchemaSequencerUtilities.h"
 #include "Sequencer/CineAssemblySchemaSpawnRegister.h"
 #include "SequencerUtilities.h"
 #include "SequencerSettings.h"
@@ -125,9 +126,9 @@ TSharedRef<SWidget> SCineAssemblySchemaWindow::BuildUI()
 	SequencerInitParams.ViewParams.UniqueName = TEXT("CineAssemblySchemaTemplateSequencer");
 
 	SequencerInitParams.ViewParams.ToolbarExtender = MakeShared<FExtender>();
+	SequencerInitParams.ViewParams.ToolbarExtender->AddToolBarExtension("BaseCommands", EExtensionHook::First, nullptr, FToolBarExtensionDelegate::CreateRaw(this, &SCineAssemblySchemaWindow::ExtendSequencerToolbarBaseCommands));
 	SequencerInitParams.ViewParams.ToolbarExtender->AddToolBarExtension("Snapping", EExtensionHook::After, nullptr, FToolBarExtensionDelegate::CreateSP(this, &SCineAssemblySchemaWindow::AddSequencerToolbarExtension));
 
-	//SequencerInitParams.HostCapabilities.bSupportsCurveEditor = true;
 	SequencerInitParams.HostCapabilities.bSupportsAddFromContentBrowser = true;
 	SequencerInitParams.HostCapabilities.bSupportsDragAndDrop = true;
 	SequencerInitParams.HostCapabilities.bSupportsSidebar = true;
@@ -1570,6 +1571,27 @@ void SCineAssemblySchemaWindow::OnActorAddedToSequencer(AActor* InActor, const F
 	}
 }
 
+void SCineAssemblySchemaWindow::ExtendSequencerToolbarBaseCommands(FToolBarBuilder& Builder)
+{
+	Builder.BeginStyleOverride("SequencerToolbar");
+
+	Builder.AddToolBarButton(
+		FUIAction(FExecuteAction::CreateLambda([this]() 
+			{
+				FCineAssemblySchemaSequencerUtilities::CreateCamera(Sequencer.ToSharedRef());
+			})),
+		"CreateCamera",
+		FText::GetEmpty(),
+		LOCTEXT("CreateCameraTooltipText", "Create a new camera and set it as the current camera cut"),
+		FSlateIcon(FAppStyle::GetAppStyleSetName(), "Sequencer.CreateCamera"),
+		EUserInterfaceActionType::Button
+	);
+
+	Builder.AddSeparator();
+
+	Builder.EndStyleOverride();
+}
+
 void SCineAssemblySchemaWindow::AddSequencerToolbarExtension(FToolBarBuilder& Builder)
 {
 	constexpr bool bShouldHaveSeparator = false;
diff --git a/Engine/Plugins/VirtualProduction/CinematicAssemblyTools/Source/CineAssemblyToolsEditor/Private/UI/CineAssembly/SCineAssemblySchemaWindow.h b/Engine/Plugins/VirtualProduction/CinematicAssemblyTools/Source/CineAssemblyToolsEditor/Private/UI/CineAssembly/SCineAssemblySchemaWindow.h
--- a/Engine/Plugins/VirtualProduction/CinematicAssemblyTools/Source/CineAssemblyToolsEditor/Private/UI/CineAssembly/SCineAssemblySchemaWindow.h
+++ b/Engine/Plugins/VirtualProduction/CinematicAssemblyTools/Source/CineAssemblyToolsEditor/Private/UI/CineAssembly/SCineAssemblySchemaWindow.h
@@ -285,6 +285,9 @@ private:
 	/** Extends the sequencer widget's toolbar */
 	void AddSequencerToolbarExtension(FToolBarBuilder& Builder);
 
+	/** Extends the Sequencer Toolbar with custom entries for this specific Sequencer */
+	void ExtendSequencerToolbarBaseCommands(FToolBarBuilder& Builder);
+
 	/** Returns the visibility of the display rate warning toolbar widget */
 	EVisibility GetDisplayRateWarningVisibility() const;
 
