ShotManagement: Add a new track type, the SubAssemblyTrack, to the Schema Sequencer for custom subsequence behavior, including:
* Adding an existing sequence as a reference
* Adding a new sequence using a Cine Assembly Schema as a template
* Adding a new sequence using an existing sequence / assembly as a template
When creating a CineAssembly from a Schema that has SubAssemblyTracks, those tracks will be replaced with normal SubTracks, and all new sequences will be created as SubAssembly assets.

#rnx
#Synced-CL 49044222

diff --git a/Engine/Plugins/VirtualProduction/CinematicAssemblyTools/Source/CineAssemblyTools/Private/CineAssembly.cpp b/Engine/Plugins/VirtualProduction/CinematicAssemblyTools/Source/CineAssemblyTools/Private/CineAssembly.cpp
--- a/Engine/Plugins/VirtualProduction/CinematicAssemblyTools/Source/CineAssemblyTools/Private/CineAssembly.cpp
+++ b/Engine/Plugins/VirtualProduction/CinematicAssemblyTools/Source/CineAssemblyTools/Private/CineAssembly.cpp
@@ -9,6 +9,7 @@
 #include "LevelSequenceShotMetaDataLibrary.h"
 #include "MovieScene.h"
 #include "MovieSceneFolder.h"
+#include "MovieSceneSubAssemblyTrack.h"
 #include "NamingTokensEngineSubsystem.h"
 #include "Sections/MovieSceneSubSection.h"
 #include "Serialization/JsonReader.h"
@@ -127,9 +128,12 @@ bool UCineAssembly::IsCompatibleAsSubSequence(const UMovieSceneSequence& ParentS
 #if WITH_EDITOR
 ETrackSupport UCineAssembly::IsTrackSupportedImpl(TSubclassOf<class UMovieSceneTrack> InTrackClass) const
 {
+	// If the UObject that owns this assembly is a Schema (rather than a UPackage, for example), then we know that this Assembly is actually a Schema Template that is only ever edited in the Schema Sequencer.
 	if (GetTypedOuter<UCineAssemblySchema>())
 	{
-		if (InTrackClass == UMovieSceneEventTrack::StaticClass())
+		// Event tracks are not supported in the Schema Sequencer
+		// Normal Subsequence Tracks are not supported in the Schema Sequencer, because the custom SubAssembly Track is used instead
+		if ((InTrackClass == UMovieSceneEventTrack::StaticClass()) || (InTrackClass == UMovieSceneSubTrack::StaticClass()))
 		{
 			return ETrackSupport::NotSupported;
 		}
@@ -154,6 +158,15 @@ void UCineAssembly::InitializeFromTemplate(ULevelSequence* Template)
 		return;
 	}
 
+	// If the template is a CineAssembly, set the schema of this assembly to match
+	if (UCineAssembly* TemplateAssembly = Cast<UCineAssembly>(Template))
+	{
+		if (const UCineAssemblySchema* TemplateSchema = TemplateAssembly->GetSchema())
+		{
+			SetSchema(const_cast<UCineAssemblySchema*>(TemplateSchema));
+		}
+	}
+
 	if (UMovieScene* TemplateMovieScene = Template->GetMovieScene())
 	{
 		// Attempt to duplicate the inner MovieScene. Assign only if duplication is a success to avoid leaving the assembly in a partially valid state.
@@ -190,6 +203,43 @@ void UCineAssembly::InitializeFromTemplate(ULevelSequence* Template)
 		}
 
 #if WITH_EDITOR
+		// Replace any SubAssemblyTracks we find with normal SubTracks (and create any SubAssemblies needed by those tracks)
+		while (UMovieSceneSubAssemblyTrack* SubAssemblyTrack = MovieScene->FindTrack<UMovieSceneSubAssemblyTrack>())
+		{
+			// Add a new SubTrack to the MovieScene. The SubAssemblyTrack will be removed from the MovieScene at the bottom of this loop.
+			UMovieSceneSubTrack* SubTrack = Cast<UMovieSceneSubTrack>(MovieScene->AddTrack(UMovieSceneSubTrack::StaticClass()));
+			SubTrack->SetDisplayName(SubAssemblyTrack->GetDisplayName());
+
+			for (UMovieSceneSection* Section : SubAssemblyTrack->GetAllSections())
+			{
+				if (UMovieSceneSubSection* SchemaSubSection = Cast<UMovieSceneSubSection>(Section))
+				{
+					// Duplicate the SubSection from the SubAssemblyTrack and add it to the new normal SubTrack
+					UMovieSceneSubSection* DuplicateSection = DuplicateObject(SchemaSubSection, SubTrack);
+					SubTrack->AddSection(*DuplicateSection);
+
+					// Check if the SubSequence in this SubSection should be added as a new sequence or as a reference, by checking its owner.
+					// A new SubSequence will be owned by the Schema. A reference would be refer to an asset that is owned by its UPackage. 
+					UMovieSceneSequence* SubSequence = SchemaSubSection->GetSequence();
+					if (SubSequence && SubSequence->GetTypedOuter<UCineAssemblySchema>())
+					{
+						// Create a new SubAssembly asset and set it as the SubSection's sequence (replacing the assembly template from the Schema).
+						UCineAssembly* AssemblyTemplate = Cast<UCineAssembly>(SubSequence);
+						UCineAssembly* NewSubAssembly = CreateSubAssemblyFromTemplate(AssemblyTemplate);
+
+						if (NewSubAssembly)
+						{
+							DuplicateSection->SetSequence(NewSubAssembly);
+							SubAssemblies.Add(DuplicateSection);
+						}
+					}
+				}
+			}
+
+			// Now remove the SubAssemblyTrack from the MovieScene. 
+			MovieScene->RemoveTrack(*SubAssemblyTrack);
+		}
+
 		// Recursively gather all of the folders in the movie scene
 		TArray<UMovieSceneFolder*> AllFolders;
 		for (UMovieSceneFolder* RootFolder : MovieScene->GetRootFolders())
@@ -233,6 +283,64 @@ void UCineAssembly::InitializeFromTemplate(ULevelSequence* Template)
 }
 
 #if WITH_EDITOR
+UCineAssembly* UCineAssembly::CreateSubAssemblyFromTemplate(UCineAssembly* AssemblyTemplate)
+{
+	if (!AssemblyTemplate)
+	{
+		return nullptr;
+	}
+
+	// Resolve any tokens found in the name of the SubAssembly name before attempting to create it 
+	FTemplateString SubAssemblyName = AssemblyTemplate->AssemblyName;
+	SubAssemblyName.Resolved = UCineAssemblyNamingTokens::GetResolvedText(SubAssemblyName.Template, AssemblyTemplate);
+
+	const FString SubAssemblyFilename = FPaths::GetBaseFilename(SubAssemblyName.Resolved.ToString());
+
+	if (SubAssemblyFilename.IsEmpty())
+	{
+		return nullptr;
+	}
+
+	// Before creating each subassembly, sanity check that each one will actually have a unique name (in case there are duplicates in the schema description)
+	FString AssemblyPath;
+	FString AssemblyRootFolder;
+	GetAssetPathAndRootFolder(AssemblyPath, AssemblyRootFolder);
+
+	IAssetTools& AssetTools = FModuleManager::GetModuleChecked<FAssetToolsModule>("AssetTools").Get();
+	FString UniquePackageName;
+	FString UniqueAssetName;
+	AssetTools.CreateUniqueAssetName(AssemblyRootFolder / SubAssemblyName.Resolved.ToString(), TEXT(""), UniquePackageName, UniqueAssetName);
+
+	// Create a new Package and CineAssembly object for the new SubAssembly
+	UPackage* SubAssemblyPackage = CreatePackage(*UniquePackageName);
+
+	const EObjectFlags Flags = RF_Public | RF_Standalone | RF_Transactional;
+	UCineAssembly* SubAssembly = NewObject<UCineAssembly>(SubAssemblyPackage, UCineAssembly::StaticClass(), FName(*UniqueAssetName), Flags);
+
+	if (SubAssembly)
+	{
+		SubAssembly->AssemblyName = SubAssemblyName;
+		SubAssembly->InitializeFromTemplate(AssemblyTemplate);
+
+		SubAssembly->Level = Level;
+		SubAssembly->ParentAssembly = this;
+		SubAssembly->Production = Production;
+		SubAssembly->ProductionName = ProductionName;
+
+		SubAssembly->CreateSubAssemblies();
+
+		ULevelSequenceShotMetaDataLibrary::SetIsSubSequence(SubAssembly, true);
+
+		// Notify the asset registry about the new subassembly
+		FAssetRegistryModule::AssetCreated(SubAssembly);
+
+		// Mark the package dirty
+		SubAssemblyPackage->MarkPackageDirty();
+	}
+
+	return SubAssembly;
+}
+
 void UCineAssembly::GetAllMovieSceneFoldersRecursive(UMovieSceneFolder* RootFolder, TArray<UMovieSceneFolder*>& OutFolders)
 {
 	for (UMovieSceneFolder* ChildFolder : RootFolder->GetChildFolders())
@@ -278,7 +386,14 @@ void UCineAssembly::ChangeSchema(UCineAssemblySchema* InSchema)
 	// Reset the assembly's name based on the schema template
 	if (BaseSchema)
 	{
-		AssemblyName.Template = BaseSchema->DefaultAssemblyName;
+		if (!BaseSchema->DefaultAssemblyName.IsEmpty())
+		{
+			AssemblyName.Template = BaseSchema->DefaultAssemblyName;
+		}
+		else
+		{
+			AssemblyName.Template = TEXT("New") + BaseSchema->SchemaName;
+		}
 	}
 	else
 	{
diff --git a/Engine/Plugins/VirtualProduction/CinematicAssemblyTools/Source/CineAssemblyTools/Private/MovieSceneSubAssemblyTrack.cpp b/Engine/Plugins/VirtualProduction/CinematicAssemblyTools/Source/CineAssemblyTools/Private/MovieSceneSubAssemblyTrack.cpp
new file mode 100644
--- /dev/null
+++ b/Engine/Plugins/VirtualProduction/CinematicAssemblyTools/Source/CineAssemblyTools/Private/MovieSceneSubAssemblyTrack.cpp
@@ -0,0 +1,8 @@
+// Copyright Epic Games, Inc. All Rights Reserved.
+
+#include "MovieSceneSubAssemblyTrack.h"
+
+UMovieSceneSubAssemblyTrack::UMovieSceneSubAssemblyTrack(const FObjectInitializer& ObjectInitializer)
+	: Super(ObjectInitializer)
+{
+}
diff --git a/Engine/Plugins/VirtualProduction/CinematicAssemblyTools/Source/CineAssemblyTools/Public/CineAssembly.h b/Engine/Plugins/VirtualProduction/CinematicAssemblyTools/Source/CineAssemblyTools/Public/CineAssembly.h
--- a/Engine/Plugins/VirtualProduction/CinematicAssemblyTools/Source/CineAssemblyTools/Public/CineAssembly.h
+++ b/Engine/Plugins/VirtualProduction/CinematicAssemblyTools/Source/CineAssemblyTools/Public/CineAssembly.h
@@ -191,6 +191,9 @@ private:
 	TOptional<FDateTime> TryGetCreatedAsLocalTime() const;
 
 #if WITH_EDITOR
+	/** Creates a new CineAssembly asset using the input assembly as a template */
+	UCineAssembly* CreateSubAssemblyFromTemplate(UCineAssembly* AssemblyTemplate);
+
 	/** Recursively gather the movie scene folders within a root folder */
 	void GetAllMovieSceneFoldersRecursive(UMovieSceneFolder* RootFolder, TArray<UMovieSceneFolder*>& OutFolders);
 #endif // WITH_EDITOR
diff --git a/Engine/Plugins/VirtualProduction/CinematicAssemblyTools/Source/CineAssemblyTools/Public/MovieSceneSubAssemblyTrack.h b/Engine/Plugins/VirtualProduction/CinematicAssemblyTools/Source/CineAssemblyTools/Public/MovieSceneSubAssemblyTrack.h
new file mode 100644
--- /dev/null
+++ b/Engine/Plugins/VirtualProduction/CinematicAssemblyTools/Source/CineAssemblyTools/Public/MovieSceneSubAssemblyTrack.h
@@ -0,0 +1,20 @@
+// Copyright Epic Games, Inc. All Rights Reserved.
+
+#pragma once
+
+#include "Tracks/MovieSceneSubTrack.h"
+
+#include "MovieSceneSubAssemblyTrack.generated.h"
+
+/**
+ * A Track to hold SubAssemblies within a parent CineAssembly.
+ * Used by the Schema Sequencer tool to support custom behavior above and beyond normal SubTracks
+ */
+UCLASS()
+class UMovieSceneSubAssemblyTrack : public UMovieSceneSubTrack
+{
+	GENERATED_BODY()
+
+public:
+	UMovieSceneSubAssemblyTrack(const FObjectInitializer& ObjectInitializer);
+};
diff --git a/Engine/Plugins/VirtualProduction/CinematicAssemblyTools/Source/CineAssemblyToolsEditor/CineAssemblyToolsEditor.build.cs b/Engine/Plugins/VirtualProduction/CinematicAssemblyTools/Source/CineAssemblyToolsEditor/CineAssemblyToolsEditor.build.cs
--- a/Engine/Plugins/VirtualProduction/CinematicAssemblyTools/Source/CineAssemblyToolsEditor/CineAssemblyToolsEditor.build.cs
+++ b/Engine/Plugins/VirtualProduction/CinematicAssemblyTools/Source/CineAssemblyToolsEditor/CineAssemblyToolsEditor.build.cs
@@ -44,6 +44,7 @@ namespace UnrealBuildTool.Rules
 					"Projects",
 					"PropertyEditor",
 					"Sequencer",
+					"SequencerCore",
 					"SharedSettingsWidgets",
 					"Slate",
 					"SlateCore",
diff --git a/Engine/Plugins/VirtualProduction/CinematicAssemblyTools/Source/CineAssemblyToolsEditor/Private/CineAssemblyToolsEditorModule.cpp b/Engine/Plugins/VirtualProduction/CinematicAssemblyTools/Source/CineAssemblyToolsEditor/Private/CineAssemblyToolsEditorModule.cpp
--- a/Engine/Plugins/VirtualProduction/CinematicAssemblyTools/Source/CineAssemblyToolsEditor/Private/CineAssemblyToolsEditorModule.cpp
+++ b/Engine/Plugins/VirtualProduction/CinematicAssemblyTools/Source/CineAssemblyToolsEditor/Private/CineAssemblyToolsEditorModule.cpp
@@ -17,6 +17,7 @@
 #include "Editor.h"
 #include "Framework/Application/SlateApplication.h"
 #include "Framework/Docking/TabManager.h"
+#include "ISequencerModule.h"
 #include "LevelEditor.h"
 #include "Misc/ConfigCacheIni.h"
 #include "MoviePipelineQueue.h"
@@ -24,6 +25,7 @@
 #include "ProductionSettings.h"
 #include "ProductionSettingsCustomization.h"
 #include "PropertyEditorModule.h"
+#include "Sequencer/SubAssemblyTrackEditor.h"
 #include "Subsystems/AssetEditorSubsystem.h"
 #include "ToolMenus.h"
 #include "UI/CineAssembly/SCineAssemblyPropertyEditor.h"
@@ -96,6 +98,10 @@ void FCineAssemblyToolsEditorModule::StartupModule()
 		const FString SchemaTabName = FString::Printf(TEXT("CineAssemblySchemaTab%d"), TabIndex);
 		ManagedSchemaTabs.Add(FTabId(*SchemaTabName), FGuid());
 	}
+
+	// Register the SubAssembly Track Editor used by the Schema Sequencer tool
+	ISequencerModule& SequencerModule = FModuleManager::Get().LoadModuleChecked<ISequencerModule>("Sequencer");
+	RegisterSubTrackEditorHandle = SequencerModule.RegisterTrackEditor(FOnCreateTrackEditor::CreateStatic(&FSubAssemblyTrackEditor::CreateTrackEditor));
 }
 
 void FCineAssemblyToolsEditorModule::RegisterProductionExtension(const UScriptStruct& DataScriptStruct)
@@ -215,6 +221,11 @@ void FCineAssemblyToolsEditorModule::ShutdownModule()
 		AssetRegistryModule->Get().OnInMemoryAssetDeleted().RemoveAll(this);
 	}
 
+	if (ISequencerModule* SequencerModule = FModuleManager::GetModulePtr<ISequencerModule>("Sequencer"))
+	{
+		SequencerModule->UnRegisterTrackEditor(RegisterSubTrackEditorHandle);
+	}
+
 	GConfig->EmptySection(*OpenTabSection, GEditorPerProjectIni);
 	SaveOpenTabs(ManagedAssemblyTabs);
 	SaveOpenTabs(ManagedSchemaTabs);
diff --git a/Engine/Plugins/VirtualProduction/CinematicAssemblyTools/Source/CineAssemblyToolsEditor/Private/CineAssemblyToolsEditorModule.h b/Engine/Plugins/VirtualProduction/CinematicAssemblyTools/Source/CineAssemblyToolsEditor/Private/CineAssemblyToolsEditorModule.h
--- a/Engine/Plugins/VirtualProduction/CinematicAssemblyTools/Source/CineAssemblyToolsEditor/Private/CineAssemblyToolsEditorModule.h
+++ b/Engine/Plugins/VirtualProduction/CinematicAssemblyTools/Source/CineAssemblyToolsEditor/Private/CineAssemblyToolsEditorModule.h
@@ -104,4 +104,5 @@ private:
 	static const FString OpenTabSection;
 
 	FDelegateHandle OnSequenceSetHandle;
+	FDelegateHandle RegisterSubTrackEditorHandle;
 };
diff --git a/Engine/Plugins/VirtualProduction/CinematicAssemblyTools/Source/CineAssemblyToolsEditor/Private/Sequencer/SubAssemblySection.cpp b/Engine/Plugins/VirtualProduction/CinematicAssemblyTools/Source/CineAssemblyToolsEditor/Private/Sequencer/SubAssemblySection.cpp
new file mode 100644
--- /dev/null
+++ b/Engine/Plugins/VirtualProduction/CinematicAssemblyTools/Source/CineAssemblyToolsEditor/Private/Sequencer/SubAssemblySection.cpp
@@ -0,0 +1,144 @@
+// Copyright Epic Games, Inc. All Rights Reserved.
+
+#include "SubAssemblySection.h"
+
+#include "CineAssembly.h"
+
+#define LOCTEXT_NAMESPACE "CineAssemblySchemaSubSection"
+
+FSubAssemblySection::FSubAssemblySection(TSharedPtr<ISequencer> InSequencer, UMovieSceneSubSection& InSectionObject)
+	: TSubSectionMixin(InSequencer, *CastChecked<UMovieSceneSubSection>(&InSectionObject))
+	, SectionObject(InSectionObject)
+{
+	// If the Section's Sequence is owned by a Schema, then this section represents a template for a new sequence, not a reference to an existing one
+	if (UMovieSceneSequence* Sequence = SectionObject.GetSequence())
+	{
+		if (Sequence->GetTypedOuter<UCineAssemblySchema>())
+		{
+			bIsTemplateSection = true;
+		}
+		else
+		{
+			bIsTemplateSection = false;
+		}
+	}
+}
+
+UMovieSceneSection* FSubAssemblySection::GetSectionObject()
+{
+	return &SectionObject;
+}
+
+FText FSubAssemblySection::GetSectionTitle() const
+{
+	// The section text is handled by the Section Widget
+	return FText::GetEmpty();
+}
+
+TSharedRef<SWidget> FSubAssemblySection::GenerateSectionWidget()
+{
+	return SNew(SHorizontalBox)
+		+ SHorizontalBox::Slot()
+		.AutoWidth()
+		.VAlign(VAlign_Top)
+		.Padding(0.0f, 2.0f, 0.0f, 0.0f)
+		[
+			// TODO: Validate that the tokenized text could possibly be a valid asset name
+			SAssignNew(EditableTextBlock, SNamingTokensEditableTextBox)
+				.SupportInlineEdit(true)
+				.ShouldEvaluateTokens(false)
+				.DisplayTokenIcon(false)
+				.DisplayBorderImage(false)
+				.IsReadOnly(!bIsTemplateSection)
+				.Text(this, &FSubAssemblySection::GetSectionText)
+				.OnTextCommitted(this, &FSubAssemblySection::SetSectionText)
+		]
+
+		+ SHorizontalBox::Slot()
+		.FillWidth(1.0f)
+		.Padding(0.0f, 5.0f, 0.0f, 0.0f)
+		[
+			SNew(STextBlock)
+				.Text(this, &FSubAssemblySection::GetOriginText)
+				.Font(FCoreStyle::GetDefaultFontStyle("Italic", 10))
+		];
+}
+
+int32 FSubAssemblySection::OnPaintSection(FSequencerSectionPainter& InPainter) const
+{
+	TSubSectionMixin<>::OnPaintSection(InPainter);
+
+	// Add a specific tint to Template Sections to more easily distinguish them from Reference Sections
+	if (bIsTemplateSection)
+	{
+		const FColor CineAssemblyAssetColor = FColor(229, 45, 113);
+		return InPainter.PaintSectionBackground(CineAssemblyAssetColor);
+	}
+
+	return InPainter.PaintSectionBackground();
+}
+
+FReply FSubAssemblySection::OnSectionDoubleClicked(const FGeometry& SectionGeometry, const FPointerEvent& MouseEvent)
+{
+	// TODO: Check the geometry to determine what double clicking should do.
+	// If double clicking on the Naming Token Textbox, it should allow the user to rename the section.
+	// If double clicking anywhere else, and the Section's sequence is a "Schema Template", then it should open that Schema in a separate window.
+
+	if (bIsTemplateSection && (SectionObject.GetSequence() != nullptr))
+	{
+		EditableTextBlock->EnterEditingMode();
+	}
+
+	return FReply::Handled();
+}
+
+FText FSubAssemblySection::GetSectionText() const
+{
+	UMovieSceneSequence* Sequence = SectionObject.GetSequence();
+	if (!Sequence)
+	{
+		return FText::GetEmpty();
+	}
+
+	if (bIsTemplateSection)
+	{
+		// Template sections will display the (potentially tokenized) name of the assembly template
+		if (UCineAssembly* CineAssembly = Cast<UCineAssembly>(Sequence))
+		{
+			return FText::FromString(CineAssembly->AssemblyName.Template);
+		}
+	}
+
+	// Reference sections simply display the name of the referenced sequence
+	return Sequence->GetDisplayName();
+}
+
+void FSubAssemblySection::SetSectionText(const FText& InText, ETextCommit::Type InCommitType)
+{
+	if (UCineAssembly* CineAssembly = Cast<UCineAssembly>(SectionObject.GetSequence()))
+	{
+		CineAssembly->AssemblyName.Template = InText.ToString();
+	}
+}
+
+FText FSubAssemblySection::GetOriginText() const
+{
+	if (bIsTemplateSection)
+	{
+		// The origin text of a template section indicates the name of the original sequence that was used as the template
+		if (UCineAssembly* CineAssembly = Cast<UCineAssembly>(SectionObject.GetSequence()))
+		{
+			FText TemplateOriginText = CineAssembly->AssemblyName.Resolved;
+			if (!TemplateOriginText.IsEmpty())
+			{
+				return FText::Format(LOCTEXT("TemplateOriginText", "from {0}"), TemplateOriginText);
+			}
+		}
+		return FText::GetEmpty();
+	}
+
+	// The origin text of a reference section simply indicates to the user that the sequence is a reference
+	return LOCTEXT("ReferenceText", "Reference");
+}
+
+#undef LOCTEXT_NAMESPACE
diff --git a/Engine/Plugins/VirtualProduction/CinematicAssemblyTools/Source/CineAssemblyToolsEditor/Private/Sequencer/SubAssemblySection.h b/Engine/Plugins/VirtualProduction/CinematicAssemblyTools/Source/CineAssemblyToolsEditor/Private/Sequencer/SubAssemblySection.h
new file mode 100644
--- /dev/null
+++ b/Engine/Plugins/VirtualProduction/CinematicAssemblyTools/Source/CineAssemblyToolsEditor/Private/Sequencer/SubAssemblySection.h
@@ -0,0 +1,47 @@
+// Copyright Epic Games, Inc. All Rights Reserved.
+
+#pragma once
+
+#include "ISequencerSection.h"
+
+#include "Sections/MovieSceneSubSection.h"
+#include "SNamingTokensEditableTextBox.h"
+#include "TrackEditors/SubTrackEditorBase.h"
+
+/**
+ * An ISequencerSection that handles the UI of sections within a SubAssembly Track
+ */
+class FSubAssemblySection	: public TSubSectionMixin<>, public TSharedFromThis<FSubAssemblySection>
+{
+public:
+	FSubAssemblySection(TSharedPtr<ISequencer> InSequencer, UMovieSceneSubSection& InSectionObject);
+
+public:
+	//~ Begin ISequencerSection interface
+	virtual UMovieSceneSection* GetSectionObject() override;
+	virtual FText GetSectionTitle() const override;
+	virtual TSharedRef<SWidget> GenerateSectionWidget() override;
+	virtual int32 OnPaintSection(FSequencerSectionPainter& InPainter) const override;
+	virtual FReply OnSectionDoubleClicked(const FGeometry& SectionGeometry, const FPointerEvent& MouseEvent) override;
+	//~ End ISequencerSection interface
+
+private:
+	/** Get the name of the subsequence in the Section to display in the Section Widget */
+	FText GetSectionText() const;
+
+	/** Set the name of the subsequence in the Section */
+	void SetSectionText(const FText& InText, ETextCommit::Type InCommitType);
+
+	/** If this is a template section, return the name of the sequence used as a template. Otherwise, return text indicating that this is a reference section. */
+	FText GetOriginText() const;
+
+private:
+	/** Section object represented by this SubSection */
+	UMovieSceneSubSection& SectionObject;
+
+	/** True if the sequence was added to this section as a Template (and not a Reference) */
+	bool bIsTemplateSection = true;
+
+	/** Editable textbox widget, used for renaming the Sequence in this Section */
+	TSharedPtr<SNamingTokensEditableTextBox> EditableTextBlock;
+};
diff --git a/Engine/Plugins/VirtualProduction/CinematicAssemblyTools/Source/CineAssemblyToolsEditor/Private/Sequencer/SubAssemblyTrackEditor.cpp b/Engine/Plugins/VirtualProduction/CinematicAssemblyTools/Source/CineAssemblyToolsEditor/Private/Sequencer/SubAssemblyTrackEditor.cpp
new file mode 100644
--- /dev/null
+++ b/Engine/Plugins/VirtualProduction/CinematicAssemblyTools/Source/CineAssemblyToolsEditor/Private/Sequencer/SubAssemblyTrackEditor.cpp
@@ -0,0 +1,287 @@
+// Copyright Epic Games, Inc. All Rights Reserved.
+
+#include "SubAssemblyTrackEditor.h"
+
+#include "CineAssembly.h"
+#include "CineAssemblySchema.h"
+#include "ContentBrowserModule.h"
+#include "IContentBrowserSingleton.h"
+#include "MovieSceneSubAssemblyTrack.h"
+#include "MovieSceneToolsProjectSettings.h"
+#include "MVVM/Extensions/ITrackExtension.h"
+#include "MVVM/Selection/Selection.h"
+#include "MVVM/Views/ViewUtilities.h"
+#include "Sections/MovieSceneSubSection.h"
+#include "SequencerSettings.h"
+#include "SequencerUtilities.h"
+#include "SubAssemblySection.h"
+#include "Widgets/Layout/SBox.h"
+
+#define LOCTEXT_NAMESPACE "SubAssemblyTrackEditor"
+
+FSubAssemblyTrackEditor::FSubAssemblyTrackEditor(TSharedRef<ISequencer> InSequencer)
+	: FSubTrackEditor(InSequencer)
+{
+}
+
+TSharedRef<ISequencerTrackEditor> FSubAssemblyTrackEditor::CreateTrackEditor(TSharedRef<ISequencer> OwningSequencer)
+{
+	return MakeShared<FSubAssemblyTrackEditor>(OwningSequencer);
+}
+
+TSubclassOf<UMovieSceneSubTrack> FSubAssemblyTrackEditor::GetSubTrackClass() const
+{
+	return UMovieSceneSubAssemblyTrack::StaticClass();
+}
+
+bool FSubAssemblyTrackEditor::SupportsSequence(UMovieSceneSequence* InSequence) const
+{
+	// If the focused Sequence that is currently open belongs to a Schema, then we know we are in the Schema Sequencer tool.
+	// In this case, CineAssemblies are always supported, and other sequences are supported if they support normal SubTracks
+	if (GetFocusedMovieScene()->GetTypedOuter<UCineAssemblySchema>())
+	{
+		if (InSequence->IsA(UCineAssembly::StaticClass()))
+		{
+			return true;
+		}
+
+		ETrackSupport TrackSupported = InSequence ? InSequence->IsTrackSupported(UMovieSceneSubTrack::StaticClass()) : ETrackSupport::NotSupported;
+		return TrackSupported == ETrackSupport::Supported;
+	}
+
+	// Outside of the Schema Sequencer, SubAssemblyTracks are not supported, and will not appear in the "Add Track" menu
+	return false;
+}
+
+TSharedRef<ISequencerSection> FSubAssemblyTrackEditor::MakeSectionInterface(UMovieSceneSection& SectionObject, UMovieSceneTrack& Track, FGuid ObjectBinding)
+{
+	UMovieSceneSubSection* SubSection = Cast<UMovieSceneSubSection>(&SectionObject);
+	return MakeShared<FSubAssemblySection>(GetSequencer(), *SubSection);
+}
+
+TSharedPtr<SWidget> FSubAssemblyTrackEditor::BuildOutlinerEditWidget(const FGuid& ObjectBinding, UMovieSceneTrack* Track, const FBuildEditWidgetParams& Params)
+{
+	return UE::Sequencer::MakeAddButton(GetSubTrackName(), FOnGetContent::CreateRaw(this, &FSubAssemblyTrackEditor::BuildAddSubSequenceComboButtonContent, Params.ViewModel.AsWeak()), Params.ViewModel);
+}
+
+TSharedRef<SWidget> FSubAssemblyTrackEditor::BuildAddSubSequenceComboButtonContent(UE::Sequencer::TWeakViewModelPtr<UE::Sequencer::FViewModel> WeakViewModel)
+{
+	using namespace UE::Sequencer;
+
+	TViewModelPtr<FViewModel> ViewModel = WeakViewModel.Pin();
+	if (!ViewModel)
+	{
+		return SNullWidget::NullWidget;
+	}
+
+	TViewModelPtr<ITrackExtension> TrackModel = ViewModel->FindAncestorOfType<ITrackExtension>(true);
+	if (!TrackModel)
+	{
+		return SNullWidget::NullWidget;
+	}
+
+	UMovieSceneSubAssemblyTrack* Track = Cast<UMovieSceneSubAssemblyTrack>(TrackModel->GetTrack());
+	if (!Track)
+	{
+		return SNullWidget::NullWidget;
+	}
+
+	constexpr bool bShouldCloseAfterMenuSelection = true;
+	FMenuBuilder MenuBuilder(bShouldCloseAfterMenuSelection, nullptr);
+
+	MenuBuilder.AddSubMenu(
+		LOCTEXT("AddNewSequence", "Add New Sequence"),
+		LOCTEXT("AddNewSequenceTooltip", "Add a new Sequence to this track using an existing Schema or Sequence as a base"),
+		FNewMenuDelegate::CreateRaw(this, &FSubAssemblyTrackEditor::AddNewSequenceSubMenu, Track));
+
+	MenuBuilder.AddSubMenu(
+		LOCTEXT("AddReference", "Add Reference to Sequence"),
+		LOCTEXT("AddReferenceTooltip", "Add a reference to an existing Sequence to this track"),
+		FNewMenuDelegate::CreateRaw(this, &FSubAssemblyTrackEditor::AddReferenceSubMenu, Track));
+
+	FSequencerUtilities::MakeTimeWarpMenuEntry(MenuBuilder, WeakViewModel);
+
+	return MenuBuilder.MakeWidget();
+}
+
+void FSubAssemblyTrackEditor::AddNewSequenceSubMenu(FMenuBuilder& MenuBuilder, UMovieSceneSubAssemblyTrack* Track)
+{
+	MenuBuilder.AddSubMenu(
+		LOCTEXT("UseSchemaTemplate", "Using Schema Template"),
+		LOCTEXT("UseSchemaTemplateTooltip", "Use an existing Schema as the base for a new Subsequence"),
+		FNewMenuDelegate::CreateLambda([this, Track](FMenuBuilder& InMenuBuilder)
+			{
+				InMenuBuilder.BeginSection(TEXT("SchemaTemplateSection"), LOCTEXT("AddNewSequenceFromSchema", "Add New Sequence Using Schema Template"));
+				{
+					// Only display CineAssemblySchemas in the asset picker
+					TArray<FTopLevelAssetPath> ClassPaths = { UCineAssemblySchema::StaticClass()->GetClassPathName() };
+
+					constexpr bool bUseAssetAsTemplate = true;
+					TSharedRef<SWidget> AssetPickerWidget = BuildAssetPicker(Track, ClassPaths, bUseAssetAsTemplate);
+
+					InMenuBuilder.AddWidget(AssetPickerWidget, FText::GetEmpty());
+				}
+				InMenuBuilder.EndSection();
+			}));
+
+	MenuBuilder.AddSubMenu(
+		LOCTEXT("UseSequenceTemplate", "Using Sequence Template"),
+		LOCTEXT("UseSequenceTemplateTooltip", "Use an existing Sequence/Assembly as the base for a new Subsequence"),
+		FNewMenuDelegate::CreateLambda([this, Track](FMenuBuilder& InMenuBuilder)
+			{
+				InMenuBuilder.BeginSection(TEXT("SequenceTemplateSection"), LOCTEXT("AddNewSequenceFromSequence", "Add New Sequence Using Sequence Template"));
+				{
+					// Only display LevelSequences and CineAssemblies in the asset picker
+					TArray<FTopLevelAssetPath> ClassPaths = { UCineAssembly::StaticClass()->GetClassPathName(), ULevelSequence::StaticClass()->GetClassPathName() };
+
+					constexpr bool bUseAssetAsTemplate = true;
+					TSharedRef<SWidget> AssetPickerWidget = BuildAssetPicker(Track, ClassPaths, bUseAssetAsTemplate);
+
+					InMenuBuilder.AddWidget(AssetPickerWidget, FText::GetEmpty());
+				}
+				InMenuBuilder.EndSection();
+			}));
+
+	MenuBuilder.AddMenuEntry(
+		LOCTEXT("EmptySequence", "Empty Sequence"),
+		LOCTEXT("EmptySequenceTooltip", "Add a new blank Sequence to this track"),
+		FSlateIcon(),
+		FUIAction(FExecuteAction::CreateRaw(this, &FSubAssemblyTrackEditor::AddEmptySequence, Track))
+	);
+}
+
+void FSubAssemblyTrackEditor::AddReferenceSubMenu(FMenuBuilder& MenuBuilder, UMovieSceneSubAssemblyTrack* Track)
+{
+	MenuBuilder.BeginSection(TEXT("SequenceReference"), LOCTEXT("AddReferenceToSequence", "Add Reference To Sequence"));
+	{
+		// Only display LevelSequences and CineAssemblies in the asset picker
+		TArray<FTopLevelAssetPath> ClassPaths = { ULevelSequence::StaticClass()->GetClassPathName() , UCineAssembly::StaticClass()->GetClassPathName() };
+
+		constexpr bool bUseAssetAsTemplate = false;
+		TSharedRef<SWidget> AssetPickerWidget = BuildAssetPicker(Track, ClassPaths, bUseAssetAsTemplate);
+
+		MenuBuilder.AddWidget(AssetPickerWidget, FText::GetEmpty());
+	}
+	MenuBuilder.EndSection();
+}
+
+TSharedRef<SWidget> FSubAssemblyTrackEditor::BuildAssetPicker(UMovieSceneSubAssemblyTrack* Track, TArray<FTopLevelAssetPath> ClassPaths, bool bUseAssetAsTemplate)
+{
+	FAssetPickerConfig AssetPickerConfig;
+	AssetPickerConfig.bAllowNullSelection = false;
+	AssetPickerConfig.bAddFilterUI = true;
+	AssetPickerConfig.bShowTypeInColumnView = false;
+	AssetPickerConfig.InitialAssetViewType = EAssetViewType::List;
+	AssetPickerConfig.Filter.bRecursiveClasses = true;
+	AssetPickerConfig.Filter.ClassPaths = ClassPaths;
+	AssetPickerConfig.SaveSettingsName = TEXT("SubAssemblyTrackAssetPicker");
+
+	if (bUseAssetAsTemplate)
+	{
+		AssetPickerConfig.OnAssetSelected = FOnAssetSelected::CreateRaw(this, &FSubAssemblyTrackEditor::OnAssetSelectedAsTemplate, Track);
+		AssetPickerConfig.OnAssetEnterPressed = FOnAssetEnterPressed::CreateRaw(this, &FSubAssemblyTrackEditor::OnAssetEnterPressedAsTemplate, Track);
+	}
+	else
+	{
+		AssetPickerConfig.OnAssetSelected = FOnAssetSelected::CreateRaw(this, &FSubAssemblyTrackEditor::OnAssetSelectedAsReference, Track);
+		AssetPickerConfig.OnAssetEnterPressed = FOnAssetEnterPressed::CreateRaw(this, &FSubAssemblyTrackEditor::OnAssetEnterPressedAsReference, Track);
+	}
+
+	TSharedPtr<ISequencer> SequencerPtr = GetSequencer();
+	const float WidthOverride = SequencerPtr.IsValid() ? SequencerPtr->GetSequencerSettings()->GetAssetBrowserWidth() : 500.f;
+	const float HeightOverride = SequencerPtr.IsValid() ? SequencerPtr->GetSequencerSettings()->GetAssetBrowserHeight() : 400.f;
+
+	FContentBrowserModule& ContentBrowserModule = FModuleManager::Get().LoadModuleChecked<FContentBrowserModule>(TEXT("ContentBrowser"));
+
+	return SNew(SBox)
+		.WidthOverride(WidthOverride)
+		.HeightOverride(HeightOverride)
+		[
+			ContentBrowserModule.Get().CreateAssetPicker(AssetPickerConfig)
+		];
+}
+
+void FSubAssemblyTrackEditor::AddEmptySequence(UMovieSceneSubAssemblyTrack* Track)
+{
+	UCineAssembly* NewAssembly = CreateCineAssemblyFromTemplate(nullptr);
+	FSubTrackEditor::AddSequenceToTrack(NewAssembly, Track);
+}
+
+void FSubAssemblyTrackEditor::OnAssetSelectedAsTemplate(const FAssetData& AssetData, UMovieSceneSubAssemblyTrack* Track)
+{
+	UCineAssembly* NewAssembly = CreateCineAssemblyFromTemplate(AssetData.GetAsset());
+	FSubTrackEditor::AddSequenceToTrack(NewAssembly, Track);
+}
+
+void FSubAssemblyTrackEditor::OnAssetEnterPressedAsTemplate(const TArray<FAssetData>& AssetData, UMovieSceneSubAssemblyTrack* Track)
+{
+	if (AssetData.Num() > 0)
+	{
+		OnAssetSelectedAsTemplate(AssetData[0], Track);
+	}
+}
+
+void FSubAssemblyTrackEditor::OnAssetSelectedAsReference(const FAssetData& AssetData, UMovieSceneSubAssemblyTrack* Track)
+{
+	if (UMovieSceneSequence* SelectedSequence = Cast<UMovieSceneSequence>(AssetData.GetAsset()))
+	{
+		FSubTrackEditor::AddSequenceToTrack(SelectedSequence, Track);
+	}
+}
+
+void FSubAssemblyTrackEditor::OnAssetEnterPressedAsReference(const TArray<FAssetData>& AssetData, UMovieSceneSubAssemblyTrack* Track)
+{
+	if (AssetData.Num() > 0)
+	{
+		OnAssetSelectedAsReference(AssetData[0], Track);
+	}
+}
+
+UCineAssembly* FSubAssemblyTrackEditor::CreateCineAssemblyFromTemplate(UObject* TemplateObject)
+{
+	// Get the Schema Template Sequence being edited by the Schema Sequencer, as well as a pointer to the Schema itself
+	UMovieSceneSequence* SchemaTemplateSequence = GetSequencer()->GetFocusedMovieSceneSequence();
+	UCineAssemblySchema* Schema = SchemaTemplateSequence->GetTypedOuter<UCineAssemblySchema>();
+
+	// Create a new Assembly, owned by the Schema, and parented to the Schema Template Sequence
+	UCineAssembly* NewAssembly = NewObject<UCineAssembly>(Schema, NAME_None);
+	NewAssembly->SetParentAssembly(Schema->TemplateSequence);
+
+	if (UCineAssemblySchema* SchemaTemplate = Cast<UCineAssemblySchema>(TemplateObject))
+	{
+		NewAssembly->Initialize();
+		NewAssembly->SetSchema(SchemaTemplate);
+		NewAssembly->AssemblyName.Resolved = FText::FromString(SchemaTemplate->SchemaName);
+	}
+	else if (UCineAssembly* AssemblyTemplate = Cast<UCineAssembly>(TemplateObject))
+	{
+		NewAssembly->InitializeFromTemplate(AssemblyTemplate);
+		NewAssembly->AssemblyName = AssemblyTemplate->AssemblyName;
+	}
+	else if (ULevelSequence* SequenceTemplate = Cast<ULevelSequence>(TemplateObject))
+	{
+		NewAssembly->InitializeFromTemplate(SequenceTemplate);
+		NewAssembly->AssemblyName.Template = SequenceTemplate->GetDisplayName().ToString();
+		NewAssembly->AssemblyName.Resolved = FText::GetEmpty();
+	}
+	else
+	{
+		NewAssembly->Initialize();
+		NewAssembly->AssemblyName.Template = TEXT("NewAssembly");
+		NewAssembly->AssemblyName.Resolved = FText::GetEmpty();
+	}
+
+	if (UMovieScene* MovieScene = NewAssembly->GetMovieScene())
+	{
+		const UMovieSceneToolsProjectSettings* MovieSceneToolsProjectSettings = GetDefault<UMovieSceneToolsProjectSettings>();
+
+		const FFrameRate TickResolution = MovieScene->GetTickResolution();
+		const FFrameNumber DefaultStartFrame = (MovieSceneToolsProjectSettings->DefaultStartTime * TickResolution).FloorToFrame();
+		const int32 DefaultDuration = (MovieSceneToolsProjectSettings->DefaultDuration * TickResolution).FloorToFrame().Value;
+		MovieScene->SetPlaybackRange(DefaultStartFrame, DefaultDuration);
+	}
+
+	return NewAssembly;
+}
+
+#undef LOCTEXT_NAMESPACE
diff --git a/Engine/Plugins/VirtualProduction/CinematicAssemblyTools/Source/CineAssemblyToolsEditor/Private/Sequencer/SubAssemblyTrackEditor.h b/Engine/Plugins/VirtualProduction/CinematicAssemblyTools/Source/CineAssemblyToolsEditor/Private/Sequencer/SubAssemblyTrackEditor.h
new file mode 100644
--- /dev/null
+++ b/Engine/Plugins/VirtualProduction/CinematicAssemblyTools/Source/CineAssemblyToolsEditor/Private/Sequencer/SubAssemblyTrackEditor.h
@@ -0,0 +1,62 @@
+// Copyright Epic Games, Inc. All Rights Reserved.
+
+#pragma once
+
+#include "TrackEditors/SubTrackEditor.h"
+
+class FMenuBuilder;
+class UCineAssembly;
+class UMovieSceneSubTrack;
+class UMovieSceneSubAssemblyTrack;
+
+struct FAssetData;
+
+/**
+ * Sequencer Track Editor for SubAssembly Tracks (extends the functionality of normal SubTracks for CineAssemblies)
+ */
+class FSubAssemblyTrackEditor : public FSubTrackEditor
+{
+public:
+	FSubAssemblyTrackEditor(TSharedRef<ISequencer> InSequencer);
+	virtual ~FSubAssemblyTrackEditor() = default;
+
+	static TSharedRef<ISequencerTrackEditor> CreateTrackEditor(TSharedRef<ISequencer> OwningSequencer);
+
+	//~ Begin ISequencerTrackEditor interface
+	virtual TSubclassOf<UMovieSceneSubTrack> GetSubTrackClass() const override;
+	virtual bool SupportsSequence(UMovieSceneSequence* InSequence) const override;
+	virtual TSharedPtr<SWidget> BuildOutlinerEditWidget(const FGuid& ObjectBinding, UMovieSceneTrack* Track, const FBuildEditWidgetParams& Params) override;
+	virtual TSharedRef<ISequencerSection> MakeSectionInterface(UMovieSceneSection& SectionObject, UMovieSceneTrack& Track, FGuid ObjectBinding) override;
+	//~ End ISequencerTrackEditor interface
+
+private:
+	/** Build the menu to display when clicking the "+" button on a SubAssembly Track */
+	TSharedRef<SWidget> BuildAddSubSequenceComboButtonContent(UE::Sequencer::TWeakViewModelPtr<UE::Sequencer::FViewModel> WeakViewModel);
+
+	/** Build an Asset Picker widget that will only display assets matching one of the input class paths */
+	TSharedRef<SWidget> BuildAssetPicker(UMovieSceneSubAssemblyTrack* Track, TArray<FTopLevelAssetPath> ClassPaths, bool bUseAssetAsTemplate);
+
+	/** Adds the "New Sequence" submenu to the SubAssembly track menu */
+	void AddNewSequenceSubMenu(FMenuBuilder& MenuBuilder, UMovieSceneSubAssemblyTrack* Track);
+
+	/** Adds the "Reference" Sequence submenu to the SubAssembly track menu */
+	void AddReferenceSubMenu(FMenuBuilder& MenuBuilder, UMovieSceneSubAssemblyTrack* Track);
+
+	/** Callback when an asset is selected in the asset picker as a template for a new sequence to add to the track */
+	void OnAssetSelectedAsTemplate(const FAssetData& AssetData, UMovieSceneSubAssemblyTrack* Track);
+
+	/** Callback when one or more assets are selected in the asset picker as a template for a new sequence to add to the track */
+	void OnAssetEnterPressedAsTemplate(const TArray<FAssetData>& AssetData, UMovieSceneSubAssemblyTrack* Track);
+
+	/** Callback when an asset is selected in the asset picker as a reference to add to the track */
+	void OnAssetSelectedAsReference(const FAssetData& AssetData, UMovieSceneSubAssemblyTrack* Track);
+
+	/** Callback when one or more assets are selected in the asset picker as a reference to add to the track */
+	void OnAssetEnterPressedAsReference(const TArray<FAssetData>& AssetData, UMovieSceneSubAssemblyTrack* Track);
+
+	/** Creates a new empty Assembly (no schema) and adds it to the Track */
+	void AddEmptySequence(UMovieSceneSubAssemblyTrack* Track);
+
+	/** Creates a new CineAssembly and initializes it using the input Template object, which can be null */
+	UCineAssembly* CreateCineAssemblyFromTemplate(UObject* TemplateObject);
+};
\ No newline at end of file
diff --git a/Engine/Plugins/VirtualProduction/CinematicAssemblyTools/Source/CineAssemblyToolsEditor/Private/UI/CineAssembly/SCineAssemblySchemaWindow.cpp b/Engine/Plugins/VirtualProduction/CinematicAssemblyTools/Source/CineAssemblyToolsEditor/Private/UI/CineAssembly/SCineAssemblySchemaWindow.cpp
--- a/Engine/Plugins/VirtualProduction/CinematicAssemblyTools/Source/CineAssemblyToolsEditor/Private/UI/CineAssembly/SCineAssemblySchemaWindow.cpp
+++ b/Engine/Plugins/VirtualProduction/CinematicAssemblyTools/Source/CineAssemblyToolsEditor/Private/UI/CineAssembly/SCineAssemblySchemaWindow.cpp
@@ -16,6 +16,7 @@
 #include "Misc/MessageDialog.h"
 #include "Modules/ModuleManager.h"
 #include "MovieScene.h"
+#include "MovieSceneSubAssemblyTrack.h"
 #include "NamingTokensEngineSubsystem.h"
 #include "ProductionSettings.h"
 #include "PropertyCustomizationHelpers.h"
@@ -1535,20 +1536,23 @@ void SCineAssemblySchemaWindow::UpdateAssetList()
 
 void SCineAssemblySchemaWindow::OnMovieSceneDataChanged(EMovieSceneDataChangeType InChangeType)
 {
-	// In order to mute cvar tracks, we only care about changes when a track is added (in case it was a cvar track)
-	// or when a track value changes (which includes the user trying to unmute the track)
+	// In order to mute specific tracks, we only care about changes when a track is added or when a track value changes (which includes the user trying to unmute the track)
 	if ((InChangeType == EMovieSceneDataChangeType::MovieSceneStructureItemAdded) || (InChangeType == EMovieSceneDataChangeType::TrackValueChanged))
 	{
 		if (UMovieScene* MovieScene = SchemaToConfigure->TemplateSequence->GetMovieScene())
 		{
-			// Because we do not have any additional information about what the actual MovieScene change was, we just loop through every track
-			// and force all of the cvar tracks to be muted
+			// Because we do not have any additional information about what the actual MovieScene change was, we just loop through every track and force all relevant tracks to be muted
 			for (UMovieSceneTrack* Track : MovieScene->GetTracks())
 			{
+				// Mute CVar tracks and Subsequence Tracks to prevent their evaluation from affecting the world
 				if (UMovieSceneCVarTrack* CVarTrack = Cast<UMovieSceneCVarTrack>(Track))
 				{
 					CVarTrack->SetLocalEvalDisabled(true);
 				}
+				else if (UMovieSceneSubAssemblyTrack* SchemaSubTrack = Cast<UMovieSceneSubAssemblyTrack>(Track))
+				{
+					SchemaSubTrack->SetLocalEvalDisabled(true);
+				}
 			}
 		}
 	}
