ShotManagement: Fix static anaylsis warning
#rnx
#Synced-CL 49044222

diff --git a/Engine/Plugins/VirtualProduction/CinematicAssemblyTools/Source/CineAssemblyToolsEditor/Private/Sequencer/SubAssemblyTrackEditor.cpp b/Engine/Plugins/VirtualProduction/CinematicAssemblyTools/Source/CineAssemblyToolsEditor/Private/Sequencer/SubAssemblyTrackEditor.cpp
--- a/Engine/Plugins/VirtualProduction/CinematicAssemblyTools/Source/CineAssemblyToolsEditor/Private/Sequencer/SubAssemblyTrackEditor.cpp
+++ b/Engine/Plugins/VirtualProduction/CinematicAssemblyTools/Source/CineAssemblyToolsEditor/Private/Sequencer/SubAssemblyTrackEditor.cpp
@@ -40,7 +40,7 @@ bool FSubAssemblyTrackEditor::SupportsSequence(UMovieSceneSequence* InSequence)
 	// In this case, CineAssemblies are always supported, and other sequences are supported if they support normal SubTracks
 	if (GetFocusedMovieScene()->GetTypedOuter<UCineAssemblySchema>())
 	{
-		if (InSequence->IsA(UCineAssembly::StaticClass()))
+		if (InSequence && InSequence->IsA(UCineAssembly::StaticClass()))
 		{
 			return true;
 		}
