ShotManagement: Set the playback range of new Cine Assemblies in Initiailize() to avoid repeating that code everywhere that calls Initialize() on a new assembly.
#rnx

#Synced-CL 49044222

diff --git a/Engine/Plugins/VirtualProduction/CinematicAssemblyTools/Source/CineAssemblyTools/Private/CineAssembly.cpp b/Engine/Plugins/VirtualProduction/CinematicAssemblyTools/Source/CineAssemblyTools/Private/CineAssembly.cpp
--- a/Engine/Plugins/VirtualProduction/CinematicAssemblyTools/Source/CineAssemblyTools/Private/CineAssembly.cpp
+++ b/Engine/Plugins/VirtualProduction/CinematicAssemblyTools/Source/CineAssemblyTools/Private/CineAssembly.cpp
@@ -23,6 +23,7 @@
 #if WITH_EDITOR
 #include "AssetToolsModule.h"
 #include "MovieSceneMetaData.h"
+#include "MovieSceneToolsProjectSettings.h"
 #endif
 
 #define LOCTEXT_NAMESPACE "CineAssembly"
@@ -147,6 +148,16 @@ ETrackSupport UCineAssembly::IsTrackSupportedImpl(TSubclassOf<class UMovieSceneT
 void UCineAssembly::Initialize()
 {
 	Super::Initialize();
+
+#if WITH_EDITOR
+	const UMovieSceneToolsProjectSettings* MovieSceneToolsProjectSettings = GetDefault<UMovieSceneToolsProjectSettings>();
+
+	const FFrameRate TickResolution = MovieScene->GetTickResolution();
+	const FFrameNumber DefaultStartFrame = (MovieSceneToolsProjectSettings->DefaultStartTime * TickResolution).FloorToFrame();
+	const int32 DefaultDuration = (MovieSceneToolsProjectSettings->DefaultDuration * TickResolution).FloorToFrame().Value;
+
+	MovieScene->SetPlaybackRange(DefaultStartFrame, DefaultDuration);
+#endif // WITH_EDITOR
 }
 
 void UCineAssembly::InitializeFromTemplate(ULevelSequence* Template)
diff --git a/Engine/Plugins/VirtualProduction/CinematicAssemblyTools/Source/CineAssemblyTools/Private/CineAssemblySchema.cpp b/Engine/Plugins/VirtualProduction/CinematicAssemblyTools/Source/CineAssemblyTools/Private/CineAssemblySchema.cpp
--- a/Engine/Plugins/VirtualProduction/CinematicAssemblyTools/Source/CineAssemblyTools/Private/CineAssemblySchema.cpp
+++ b/Engine/Plugins/VirtualProduction/CinematicAssemblyTools/Source/CineAssemblyTools/Private/CineAssemblySchema.cpp
@@ -10,7 +10,6 @@
 
 #if WITH_EDITOR
 #include "AssetToolsModule.h"
-#include "MovieSceneToolsProjectSettings.h"
 #endif
 
 const FName UCineAssemblySchema::SchemaGuidPropertyName = GET_MEMBER_NAME_CHECKED(UCineAssemblySchema, SchemaGuid);
@@ -46,13 +45,6 @@ void UCineAssemblySchema::PostInitProperties()
 	{
 		TemplateSequence = NewObject<UCineAssembly>(this);
 		TemplateSequence->Initialize();
-
-#if WITH_EDITOR
-		const UMovieSceneToolsProjectSettings* MovieSceneSettings = GetDefault<UMovieSceneToolsProjectSettings>();
-
-		FFrameRate TickResolution = TemplateSequence->GetMovieScene()->GetTickResolution();
-		TemplateSequence->GetMovieScene()->SetPlaybackRange((MovieSceneSettings->DefaultStartTime * TickResolution).FloorToFrame(), (MovieSceneSettings->DefaultDuration * TickResolution).FloorToFrame().Value);
-#endif // WITH_EDITOR
 	}
 }
 
diff --git a/Engine/Plugins/VirtualProduction/CinematicAssemblyTools/Source/CineAssemblyToolsEditor/Private/CineAssemblyFactory.cpp b/Engine/Plugins/VirtualProduction/CinematicAssemblyTools/Source/CineAssemblyToolsEditor/Private/CineAssemblyFactory.cpp
--- a/Engine/Plugins/VirtualProduction/CinematicAssemblyTools/Source/CineAssemblyToolsEditor/Private/CineAssemblyFactory.cpp
+++ b/Engine/Plugins/VirtualProduction/CinematicAssemblyTools/Source/CineAssemblyToolsEditor/Private/CineAssemblyFactory.cpp
@@ -58,16 +58,6 @@ UObject* UCineAssemblyFactory::FactoryCreateNew(UClass* Class, UObject* InParent
 	UCineAssembly* NewAssembly = NewObject<UCineAssembly>(InParent, Name, Flags);
 	NewAssembly->Initialize();
 
-	if (UMovieScene* MovieScene = NewAssembly->GetMovieScene())
-	{
-		const UMovieSceneToolsProjectSettings* MovieSceneToolsProjectSettings = GetDefault<UMovieSceneToolsProjectSettings>();
-
-		const FFrameRate TickResolution = MovieScene->GetTickResolution();
-		const FFrameNumber DefaultStartFrame = (MovieSceneToolsProjectSettings->DefaultStartTime * TickResolution).FloorToFrame();
-		const int32 DefaultDuration = (MovieSceneToolsProjectSettings->DefaultDuration * TickResolution).FloorToFrame().Value;
-		MovieScene->SetPlaybackRange(DefaultStartFrame, DefaultDuration);
-	}
-
 	return NewAssembly;
 }
 
@@ -178,16 +168,6 @@ void UCineAssemblyFactory::CreateConfiguredAssembly(UCineAssembly* ConfiguredAss
 		{
 			// Do the same setup that ULevelSequence assets do when they are created
 			ConfiguredAssembly->Initialize();
-
-			if (UMovieScene* MovieScene = ConfiguredAssembly->GetMovieScene())
-			{
-				const UMovieSceneToolsProjectSettings* MovieSceneToolsProjectSettings = GetDefault<UMovieSceneToolsProjectSettings>();
-
-				const FFrameRate TickResolution = MovieScene->GetTickResolution();
-				const FFrameNumber DefaultStartFrame = (MovieSceneToolsProjectSettings->DefaultStartTime * TickResolution).FloorToFrame();
-				const int32 DefaultDuration = (MovieSceneToolsProjectSettings->DefaultDuration * TickResolution).FloorToFrame().Value;
-				MovieScene->SetPlaybackRange(DefaultStartFrame, DefaultDuration);
-			}
 		}
 	}
 
diff --git a/Engine/Plugins/VirtualProduction/CinematicAssemblyTools/Source/CineAssemblyToolsEditor/Private/Sequencer/SubAssemblyTrackEditor.cpp b/Engine/Plugins/VirtualProduction/CinematicAssemblyTools/Source/CineAssemblyToolsEditor/Private/Sequencer/SubAssemblyTrackEditor.cpp
--- a/Engine/Plugins/VirtualProduction/CinematicAssemblyTools/Source/CineAssemblyToolsEditor/Private/Sequencer/SubAssemblyTrackEditor.cpp
+++ b/Engine/Plugins/VirtualProduction/CinematicAssemblyTools/Source/CineAssemblyToolsEditor/Private/Sequencer/SubAssemblyTrackEditor.cpp
@@ -271,16 +271,6 @@ UCineAssembly* FSubAssemblyTrackEditor::CreateCineAssemblyFromTemplate(UObject*
 		NewAssembly->AssemblyName.Resolved = FText::GetEmpty();
 	}
 
-	if (UMovieScene* MovieScene = NewAssembly->GetMovieScene())
-	{
-		const UMovieSceneToolsProjectSettings* MovieSceneToolsProjectSettings = GetDefault<UMovieSceneToolsProjectSettings>();
-
-		const FFrameRate TickResolution = MovieScene->GetTickResolution();
-		const FFrameNumber DefaultStartFrame = (MovieSceneToolsProjectSettings->DefaultStartTime * TickResolution).FloorToFrame();
-		const int32 DefaultDuration = (MovieSceneToolsProjectSettings->DefaultDuration * TickResolution).FloorToFrame().Value;
-		MovieScene->SetPlaybackRange(DefaultStartFrame, DefaultDuration);
-	}
-
 	return NewAssembly;
 }
 
