ShotManagement: Add new MovieSceneSubAssemblySection, which derives from MovieSceneSubSection, so that the section itself can manage the template asset used to create new SubAssemblies.
Add the ability to change a SubAssemblySection between a Reference and a Template, and to change the Template object via the details panel.
The SubAssemblyTrackEditor now no longer uses the template when adding a new section to the track. Instead it simply creates an empty assembly.

#Synced-CL 49044222

diff --git a/Engine/Plugins/VirtualProduction/CinematicAssemblyTools/Source/CineAssemblyTools/Private/CineAssembly.cpp b/Engine/Plugins/VirtualProduction/CinematicAssemblyTools/Source/CineAssemblyTools/Private/CineAssembly.cpp
--- a/Engine/Plugins/VirtualProduction/CinematicAssemblyTools/Source/CineAssemblyTools/Private/CineAssembly.cpp
+++ b/Engine/Plugins/VirtualProduction/CinematicAssemblyTools/Source/CineAssemblyTools/Private/CineAssembly.cpp
@@ -9,6 +9,7 @@
 #include "LevelSequenceShotMetaDataLibrary.h"
 #include "MovieScene.h"
 #include "MovieSceneFolder.h"
+#include "MovieSceneSubAssemblySection.h"
 #include "MovieSceneSubAssemblyTrack.h"
 #include "NamingTokensEngineSubsystem.h"
 #include "Sections/MovieSceneSubSection.h"
@@ -223,25 +224,20 @@ void UCineAssembly::InitializeFromTemplate(ULevelSequence* Template)
 
 			for (UMovieSceneSection* Section : SubAssemblyTrack->GetAllSections())
 			{
-				if (UMovieSceneSubSection* SchemaSubSection = Cast<UMovieSceneSubSection>(Section))
+				if (UMovieSceneSubAssemblySection* SubAssemblySection = Cast<UMovieSceneSubAssemblySection>(Section))
 				{
-					// Duplicate the SubSection from the SubAssemblyTrack and add it to the new normal SubTrack
-					UMovieSceneSubSection* DuplicateSection = DuplicateObject(SchemaSubSection, SubTrack);
-					SubTrack->AddSection(*DuplicateSection);
-
-					// Check if the SubSequence in this SubSection should be added as a new sequence or as a reference, by checking its owner.
-					// A new SubSequence will be owned by the Schema. A reference would be refer to an asset that is owned by its UPackage. 
-					UMovieSceneSequence* SubSequence = SchemaSubSection->GetSequence();
-					if (SubSequence && SubSequence->GetTypedOuter<UCineAssemblySchema>())
-					{
-						// Create a new SubAssembly asset and set it as the SubSection's sequence (replacing the assembly template from the Schema).
-						UCineAssembly* AssemblyTemplate = Cast<UCineAssembly>(SubSequence);
-						UCineAssembly* NewSubAssembly = CreateSubAssemblyFromTemplate(AssemblyTemplate);
+					// Create a new SubSection, with the relevant properties copied from the SubAssemblySection, and add it to the new SubTrack
+					UMovieSceneSubSection* NewSubSection = NewObject<UMovieSceneSubSection>(SubTrack, NAME_None, SubAssemblySection->GetFlags(), SubAssemblySection);
+					SubTrack->AddSection(*NewSubSection);
 
+					if (SubAssemblySection->IsTemplateSection())
+					{
+						// Create a new SubAssembly using the assembly template specified by the SubAssemblySection
+						UCineAssembly* NewSubAssembly = CreateSubAssemblyFromTemplate(SubAssemblySection->NewSequenceName.ToString(), SubAssemblySection->GetAssemblyTemplate());
 						if (NewSubAssembly)
 						{
-							DuplicateSection->SetSequence(NewSubAssembly);
-							SubAssemblies.Add(DuplicateSection);
+							NewSubSection->SetSequence(NewSubAssembly);
+							SubAssemblies.Add(NewSubSection);
 						}
 					}
 				}
@@ -294,25 +290,43 @@ void UCineAssembly::InitializeFromTemplate(ULevelSequence* Template)
 }
 
 #if WITH_EDITOR
-UCineAssembly* UCineAssembly::CreateSubAssemblyFromTemplate(UCineAssembly* AssemblyTemplate)
+UCineAssembly* UCineAssembly::CreateSubAssemblyFromTemplate(const FString& SubAssemblyName, UObject* TemplateObject)
 {
-	if (!AssemblyTemplate)
+	// The new SubAssembly begins are a transient object so that it can be configured before it is named and given a package.
+	// This allows us to set its metadata first before trying to resolve any tokens which might appear in the SubAssemblyName.
+	UCineAssembly* SubAssembly = NewObject<UCineAssembly>(GetTransientPackage(), NAME_None, RF_Public | RF_Standalone | RF_Transactional);
+
+	// Initialize the new SubAssembly using the TemplateObject
+	if (UCineAssemblySchema* SchemaTemplate = Cast<UCineAssemblySchema>(TemplateObject))
+	{
+		SubAssembly->Initialize();
+		SubAssembly->SetSchema(SchemaTemplate);
+	}
+	else if (ULevelSequence* SequenceTemplate = Cast<ULevelSequence>(TemplateObject))
 	{
-		return nullptr;
+		SubAssembly->InitializeFromTemplate(SequenceTemplate);
+	}
+	else
+	{
+		SubAssembly->Initialize();
 	}
 
-	// Resolve any tokens found in the name of the SubAssembly name before attempting to create it 
-	FTemplateString SubAssemblyName = AssemblyTemplate->AssemblyName;
-	SubAssemblyName.Resolved = UCineAssemblyNamingTokens::GetResolvedText(SubAssemblyName.Template, AssemblyTemplate);
+	SubAssembly->Level = Level;
+	SubAssembly->ParentAssembly = this;
+	SubAssembly->Production = Production;
+	SubAssembly->ProductionName = ProductionName;
 
-	const FString SubAssemblyFilename = FPaths::GetBaseFilename(SubAssemblyName.Resolved.ToString());
+	// Evaluate the name of the assembly
+	SubAssembly->AssemblyName.Template = SubAssemblyName;
+	SubAssembly->AssemblyName.Resolved = UCineAssemblyNamingTokens::GetResolvedText(SubAssembly->AssemblyName.Template, SubAssembly);
 
-	if (SubAssemblyFilename.IsEmpty())
+	// If the assembly name is empty or invalid, assign it a valid default name
+	if (SubAssembly->AssemblyName.Resolved.IsEmpty() || FName(*SubAssembly->AssemblyName.Resolved.ToString()).IsNone())
 	{
-		return nullptr;
+		SubAssembly->AssemblyName.Resolved = LOCTEXT("NewCineAssemblyName", "NewCineAssembly");
+		SubAssembly->AssemblyName.Template = SubAssembly->AssemblyName.Resolved.ToString();
 	}
 
-	// Before creating each subassembly, sanity check that each one will actually have a unique name (in case there are duplicates in the schema description)
 	FString AssemblyPath;
 	FString AssemblyRootFolder;
 	GetAssetPathAndRootFolder(AssemblyPath, AssemblyRootFolder);
@@ -320,34 +334,19 @@ UCineAssembly* UCineAssembly::CreateSubAssemblyFromTemplate(UCineAssembly* Assem
 	IAssetTools& AssetTools = FModuleManager::GetModuleChecked<FAssetToolsModule>("AssetTools").Get();
 	FString UniquePackageName;
 	FString UniqueAssetName;
-	AssetTools.CreateUniqueAssetName(AssemblyRootFolder / SubAssemblyName.Resolved.ToString(), TEXT(""), UniquePackageName, UniqueAssetName);
-
-	// Create a new Package and CineAssembly object for the new SubAssembly
-	UPackage* SubAssemblyPackage = CreatePackage(*UniquePackageName);
+	AssetTools.CreateUniqueAssetName(AssemblyRootFolder / SubAssembly->AssemblyName.Resolved.ToString(), TEXT(""), UniquePackageName, UniqueAssetName);
 
-	const EObjectFlags Flags = RF_Public | RF_Standalone | RF_Transactional;
-	UCineAssembly* SubAssembly = NewObject<UCineAssembly>(SubAssemblyPackage, UCineAssembly::StaticClass(), FName(*UniqueAssetName), Flags);
+	// Now, we can create a real package for it and rename it
+	UPackage* Package = CreatePackage(*UniquePackageName);
+	SubAssembly->Rename(*UniqueAssetName, Package);
 
-	if (SubAssembly)
-	{
-		SubAssembly->AssemblyName = SubAssemblyName;
-		SubAssembly->InitializeFromTemplate(AssemblyTemplate);
-
-		SubAssembly->Level = Level;
-		SubAssembly->ParentAssembly = this;
-		SubAssembly->Production = Production;
-		SubAssembly->ProductionName = ProductionName;
-
-		SubAssembly->CreateSubAssemblies();
+	SubAssembly->CreateSubAssemblies();
 
-		ULevelSequenceShotMetaDataLibrary::SetIsSubSequence(SubAssembly, true);
+	// Notify the asset registry about the new assembly
+	FAssetRegistryModule::AssetCreated(SubAssembly);
 
-		// Notify the asset registry about the new subassembly
-		FAssetRegistryModule::AssetCreated(SubAssembly);
-
-		// Mark the package dirty
-		SubAssemblyPackage->MarkPackageDirty();
-	}
+	// Mark the package dirty
+	Package->MarkPackageDirty();
 
 	return SubAssembly;
 }
@@ -397,14 +396,7 @@ void UCineAssembly::ChangeSchema(UCineAssemblySchema* InSchema)
 	// Reset the assembly's name based on the schema template
 	if (BaseSchema)
 	{
-		if (!BaseSchema->DefaultAssemblyName.IsEmpty())
-		{
-			AssemblyName.Template = BaseSchema->DefaultAssemblyName;
-		}
-		else
-		{
-			AssemblyName.Template = TEXT("New") + BaseSchema->SchemaName;
-		}
+		AssemblyName.Template = BaseSchema->GetDefaultAssemblyName();
 	}
 	else
 	{
diff --git a/Engine/Plugins/VirtualProduction/CinematicAssemblyTools/Source/CineAssemblyTools/Private/CineAssemblySchema.cpp b/Engine/Plugins/VirtualProduction/CinematicAssemblyTools/Source/CineAssemblyTools/Private/CineAssemblySchema.cpp
--- a/Engine/Plugins/VirtualProduction/CinematicAssemblyTools/Source/CineAssemblyTools/Private/CineAssemblySchema.cpp
+++ b/Engine/Plugins/VirtualProduction/CinematicAssemblyTools/Source/CineAssemblyTools/Private/CineAssemblySchema.cpp
@@ -71,6 +71,15 @@ void UCineAssemblySchema::PostLoad()
 	}
 }
 
+FString UCineAssemblySchema::GetDefaultAssemblyName() const
+{
+	if (DefaultAssemblyName.IsEmpty())
+	{
+		return TEXT("New") + SchemaName;
+	}
+	return DefaultAssemblyName;
+}
+
 bool UCineAssemblySchema::SupportsRename() const
 {
 	return bSupportsRename;
diff --git a/Engine/Plugins/VirtualProduction/CinematicAssemblyTools/Source/CineAssemblyTools/Private/MovieSceneSubAssemblySection.cpp b/Engine/Plugins/VirtualProduction/CinematicAssemblyTools/Source/CineAssemblyTools/Private/MovieSceneSubAssemblySection.cpp
new file mode 100644
--- /dev/null
+++ b/Engine/Plugins/VirtualProduction/CinematicAssemblyTools/Source/CineAssemblyTools/Private/MovieSceneSubAssemblySection.cpp
@@ -0,0 +1,91 @@
+// Copyright Epic Games, Inc. All Rights Reserved.
+
+#include "MovieSceneSubAssemblySection.h"
+
+#include "CineAssembly.h"
+#include "CineAssemblySchema.h"
+#include "LevelSequence.h"
+
+#define LOCTEXT_NAMESPACE "MovieSceneSubAssemblySection"
+
+const FName UMovieSceneSubAssemblySection::AssemblyTemplatePropertyName = GET_MEMBER_NAME_CHECKED(UMovieSceneSubAssemblySection, AssemblyTemplate);
+const FColor UMovieSceneSubAssemblySection::ReferenceSectionColorTint = FColor(); // No Tint
+const FColor UMovieSceneSubAssemblySection::TemplateSectionColorTint = FColor(229, 45, 113); // The CineAssembly asset color
+
+UMovieSceneSubAssemblySection::UMovieSceneSubAssemblySection(const FObjectInitializer& ObjectInitializer)
+	: Super(ObjectInitializer)
+{
+}
+
+#if WITH_EDITOR
+void UMovieSceneSubAssemblySection::SetAssemblyTemplate(UObject* TemplateObject)
+{
+	SectionType = ESubAssemblySectionType::Template;
+
+	// Reset the section tint (unless the user has manually changed it to something else)
+	if (GetColorTint() == ReferenceSectionColorTint)
+	{
+		SetColorTint(TemplateSectionColorTint);
+	}
+
+	// Early-out if the template object has not actually changed to avoid unintentionally resetting the NewSequenceName
+	if (TemplateObject && (AssemblyTemplate == TemplateObject))
+	{
+		return;
+	}
+
+	AssemblyTemplate = TemplateObject;
+
+	// Update the name of the NewSequence that this section will create based on the new template object
+	if (UCineAssemblySchema* SchemaTemplate = Cast<UCineAssemblySchema>(TemplateObject))
+	{
+		NewSequenceName = FText::FromString(SchemaTemplate->GetDefaultAssemblyName());
+	}
+	else if (UCineAssembly* TemplateAssembly = Cast<UCineAssembly>(TemplateObject))
+	{	
+		NewSequenceName = FText::FromString(TemplateAssembly->AssemblyName.Template);
+	}
+	else if (ULevelSequence* SequenceTemplate = Cast<ULevelSequence>(TemplateObject))
+	{
+		NewSequenceName = SequenceTemplate->GetDisplayName();
+	}
+	else
+	{
+		NewSequenceName = LOCTEXT("NewSubSequenceName", "NewSubSequence");
+	}
+}
+
+void UMovieSceneSubAssemblySection::PostEditChangeProperty(FPropertyChangedEvent& PropertyChangedEvent)
+{
+	Super::PostEditChangeProperty(PropertyChangedEvent);
+
+	const FName PropertyName = PropertyChangedEvent.GetPropertyName();
+
+	if (PropertyName == GET_MEMBER_NAME_CHECKED(UMovieSceneSubAssemblySection, SectionType))
+	{
+		if (SectionType == ESubAssemblySectionType::Template)
+		{
+			// The section was previously a Reference section, so set the Assembly Template to be the SubSequence that was previously referenced
+			SetAssemblyTemplate(SubSequence);
+		}
+		else
+		{
+			// The section was previously a Template section, so set the SubSequence that will be referenced to be the Assembly Template object
+			// Note: This is only valid if the Assembly Template is a valid sequence (i.e. not null, and not a Schema).
+			// The detail customization disables the SectionType property in cases where this would not be valid.
+			if (UMovieSceneSequence* Sequence = Cast<UMovieSceneSequence>(AssemblyTemplate))
+			{
+				SetSequence(Sequence);
+
+				// Reset the section tint (unless the user has manually changed it to something else)
+				if (GetColorTint() == TemplateSectionColorTint)
+				{
+					SetColorTint(ReferenceSectionColorTint);
+				}
+			}
+		}
+	}
+}
+#endif // WITH_EDITOR
+
+#undef LOCTEXT_NAMESPACE
diff --git a/Engine/Plugins/VirtualProduction/CinematicAssemblyTools/Source/CineAssemblyTools/Private/MovieSceneSubAssemblyTrack.cpp b/Engine/Plugins/VirtualProduction/CinematicAssemblyTools/Source/CineAssemblyTools/Private/MovieSceneSubAssemblyTrack.cpp
--- a/Engine/Plugins/VirtualProduction/CinematicAssemblyTools/Source/CineAssemblyTools/Private/MovieSceneSubAssemblyTrack.cpp
+++ b/Engine/Plugins/VirtualProduction/CinematicAssemblyTools/Source/CineAssemblyTools/Private/MovieSceneSubAssemblyTrack.cpp
@@ -2,7 +2,14 @@
 
 #include "MovieSceneSubAssemblyTrack.h"
 
+#include "MovieSceneSubAssemblySection.h"
+
 UMovieSceneSubAssemblyTrack::UMovieSceneSubAssemblyTrack(const FObjectInitializer& ObjectInitializer)
 	: Super(ObjectInitializer)
 {
 }
+
+UMovieSceneSection* UMovieSceneSubAssemblyTrack::CreateNewSection()
+{
+	return NewObject<UMovieSceneSubAssemblySection>(this, NAME_None, RF_Transactional);
+}
diff --git a/Engine/Plugins/VirtualProduction/CinematicAssemblyTools/Source/CineAssemblyTools/Public/CineAssembly.h b/Engine/Plugins/VirtualProduction/CinematicAssemblyTools/Source/CineAssemblyTools/Public/CineAssembly.h
--- a/Engine/Plugins/VirtualProduction/CinematicAssemblyTools/Source/CineAssemblyTools/Public/CineAssembly.h
+++ b/Engine/Plugins/VirtualProduction/CinematicAssemblyTools/Source/CineAssemblyTools/Public/CineAssembly.h
@@ -191,8 +191,8 @@ private:
 	TOptional<FDateTime> TryGetCreatedAsLocalTime() const;
 
 #if WITH_EDITOR
-	/** Creates a new CineAssembly asset using the input assembly as a template */
-	UCineAssembly* CreateSubAssemblyFromTemplate(UCineAssembly* AssemblyTemplate);
+	/** Creates a new CineAssembly asset using the input object as a template */
+	UCineAssembly* CreateSubAssemblyFromTemplate(const FString& SubAssemblyName, UObject* TemplateObject);
 
 	/** Recursively gather the movie scene folders within a root folder */
 	void GetAllMovieSceneFoldersRecursive(UMovieSceneFolder* RootFolder, TArray<UMovieSceneFolder*>& OutFolders);
diff --git a/Engine/Plugins/VirtualProduction/CinematicAssemblyTools/Source/CineAssemblyTools/Public/CineAssemblySchema.h b/Engine/Plugins/VirtualProduction/CinematicAssemblyTools/Source/CineAssemblyTools/Public/CineAssemblySchema.h
--- a/Engine/Plugins/VirtualProduction/CinematicAssemblyTools/Source/CineAssemblyTools/Public/CineAssemblySchema.h
+++ b/Engine/Plugins/VirtualProduction/CinematicAssemblyTools/Source/CineAssemblyTools/Public/CineAssemblySchema.h
@@ -73,6 +73,9 @@ public:
 	/** Get the unique ID of this schema */
 	CINEASSEMBLYTOOLS_API FGuid GetSchemaGuid() const;
 
+	/** Get the default name for a new Assembly made from this schema */
+	CINEASSEMBLYTOOLS_API FString GetDefaultAssemblyName() const;
+
 	/** Whether or not this schema can be renamed */
 	CINEASSEMBLYTOOLS_API bool SupportsRename() const;
 
diff --git a/Engine/Plugins/VirtualProduction/CinematicAssemblyTools/Source/CineAssemblyTools/Public/MovieSceneSubAssemblySection.h b/Engine/Plugins/VirtualProduction/CinematicAssemblyTools/Source/CineAssemblyTools/Public/MovieSceneSubAssemblySection.h
new file mode 100644
--- /dev/null
+++ b/Engine/Plugins/VirtualProduction/CinematicAssemblyTools/Source/CineAssemblyTools/Public/MovieSceneSubAssemblySection.h
@@ -0,0 +1,72 @@
+// Copyright Epic Games, Inc. All Rights Reserved.
+
+#pragma once
+
+#include "Sections/MovieSceneSubSection.h"
+
+#include "MovieSceneSubAssemblySection.generated.h"
+
+#define UE_API CINEASSEMBLYTOOLS_API
+
+/** Indicates how a SubAssemblySection will be used */
+UENUM()
+enum class ESubAssemblySectionType : uint8
+{
+	Reference, /** Section contains a reference to an existing sequence */
+	Template   /** Section contains a template object used to create a new sequence */
+};
+
+/**
+ * Implements a section in SubAssembly Tracks
+ */
+UCLASS(BlueprintType, MinimalAPI)
+class UMovieSceneSubAssemblySection : public UMovieSceneSubSection
+{
+	GENERATED_BODY()
+
+public:
+	UMovieSceneSubAssemblySection(const FObjectInitializer& ObjectInitializer);
+
+#if WITH_EDITOR
+	//~ Begin UObject interface
+	virtual void PostEditChangeProperty(FPropertyChangedEvent& PropertyChangedEvent) override;
+	//~ End UObject interface
+#endif //WITH_EDITOR
+
+	/** Returns true if this section is used as a reference to an existing sequence */
+	UE_API bool IsReferenceSection() { return SectionType == ESubAssemblySectionType::Reference; }
+
+	/** Returns true if this section is used as a template to create a new sequence */
+	UE_API bool IsTemplateSection() { return SectionType == ESubAssemblySectionType::Template; }
+
+	/** Returns the Assembly Template object */
+	UE_API UObject* GetAssemblyTemplate() { return AssemblyTemplate; }
+
+#if WITH_EDITOR
+	/** Sets the Assembly Template and updates the SequenceName to match the new template */
+	UE_API void SetAssemblyTemplate(UObject* TemplateObject);
+#endif //WITH_EDITOR
+
+public:
+	/** Indicates whether this section will be used as a reference to an existing sequence or as a template to create a new sequence */
+	UPROPERTY(EditAnywhere, Category = "Sequence")
+	ESubAssemblySectionType SectionType = ESubAssemblySectionType::Reference;
+
+	/** The name that will be used when creating a new sequence for this section. Only applicable to Template sections */
+	UPROPERTY()
+	FText NewSequenceName;
+
+	/** Name of the protected AssemblyTemplate property */
+	static UE_API const FName AssemblyTemplatePropertyName;
+
+protected:
+	/** The object that will be used as a template to create a new sequence in this section. Can be a Level Sequence, Cine Assembly, or Schema */
+	UPROPERTY(EditAnywhere, Category = "Sequence")
+	TObjectPtr<UObject> AssemblyTemplate;
+
+private:
+	static const FColor ReferenceSectionColorTint;
+	static const FColor TemplateSectionColorTint;
+};
+
+#undef UE_API
diff --git a/Engine/Plugins/VirtualProduction/CinematicAssemblyTools/Source/CineAssemblyTools/Public/MovieSceneSubAssemblyTrack.h b/Engine/Plugins/VirtualProduction/CinematicAssemblyTools/Source/CineAssemblyTools/Public/MovieSceneSubAssemblyTrack.h
--- a/Engine/Plugins/VirtualProduction/CinematicAssemblyTools/Source/CineAssemblyTools/Public/MovieSceneSubAssemblyTrack.h
+++ b/Engine/Plugins/VirtualProduction/CinematicAssemblyTools/Source/CineAssemblyTools/Public/MovieSceneSubAssemblyTrack.h
@@ -6,15 +6,23 @@
 
 #include "MovieSceneSubAssemblyTrack.generated.h"
 
+#define UE_API CINEASSEMBLYTOOLS_API
+
 /**
  * A Track to hold SubAssemblies within a parent CineAssembly.
  * Used by the Schema Sequencer tool to support custom behavior above and beyond normal SubTracks
  */
-UCLASS()
+UCLASS(BlueprintType, MinimalAPI)
 class UMovieSceneSubAssemblyTrack : public UMovieSceneSubTrack
 {
 	GENERATED_BODY()
 
 public:
 	UMovieSceneSubAssemblyTrack(const FObjectInitializer& ObjectInitializer);
+
+	//~ Begin UMovieSceneTrack interface
+	UE_API virtual UMovieSceneSection* CreateNewSection() override;
+	//~ End UMovieSceneTrack interface
 };
+
+#undef UE_API
diff --git a/Engine/Plugins/VirtualProduction/CinematicAssemblyTools/Source/CineAssemblyToolsEditor/Private/Sequencer/MovieSceneSubAssemblySectionCustomization.cpp b/Engine/Plugins/VirtualProduction/CinematicAssemblyTools/Source/CineAssemblyToolsEditor/Private/Sequencer/MovieSceneSubAssemblySectionCustomization.cpp
new file mode 100644
--- /dev/null
+++ b/Engine/Plugins/VirtualProduction/CinematicAssemblyTools/Source/CineAssemblyToolsEditor/Private/Sequencer/MovieSceneSubAssemblySectionCustomization.cpp
@@ -0,0 +1,94 @@
+// Copyright Epic Games, Inc. All Rights Reserved.
+
+#include "MovieSceneSubAssemblySectionCustomization.h"
+
+#include "CineAssembly.h"
+#include "CineAssemblySchema.h"
+#include "DetailCategoryBuilder.h"
+#include "DetailLayoutBuilder.h"
+#include "DetailWidgetRow.h"
+#include "LevelSequence.h"
+#include "MovieSceneSubAssemblySection.h"
+#include "PropertyCustomizationHelpers.h"
+
+TSharedRef<IDetailCustomization> FSubAssemblySectionDetailCustomization::MakeInstance()
+{
+	return MakeShared<FSubAssemblySectionDetailCustomization>();
+}
+
+void FSubAssemblySectionDetailCustomization::CustomizeDetails(IDetailLayoutBuilder& DetailBuilder)
+{
+	TArray<TWeakObjectPtr<>> CustomizedObjects;
+	DetailBuilder.GetObjectsBeingCustomized(CustomizedObjects);
+	if (CustomizedObjects.Num() != 1)
+	{
+		return;
+	}
+
+	UMovieSceneSubAssemblySection* SubAssemblySection = Cast<UMovieSceneSubAssemblySection>(CustomizedObjects[0]);
+	if (!SubAssemblySection)
+	{
+		return;
+	}
+
+	IDetailCategoryBuilder& SequenceCategory = DetailBuilder.EditCategory("Sequence");
+
+	// Add the SectionType property to the Sequence Category, but disable it when the Section's Assembly Template is not a valid Sequence.
+	// This could mean that the Assembly Template is null, or possibly a Schema. 
+	// In either of these cases, it does not make sense for the user to change the SectionType from Template to Reference, because Reference Sections must refer to a valid Sequence.
+	TSharedRef<IPropertyHandle> SectionTypePropertyHandle = DetailBuilder.GetProperty(GET_MEMBER_NAME_CHECKED(UMovieSceneSubAssemblySection, SectionType));
+	IDetailPropertyRow& SectionTypePropertyRow = SequenceCategory.AddProperty(SectionTypePropertyHandle);
+	SectionTypePropertyRow.IsEnabled(TAttribute<bool>::CreateLambda([SubAssemblySection]()
+		{
+			UMovieSceneSequence* Sequence = Cast<UMovieSceneSequence>(SubAssemblySection->GetAssemblyTemplate());
+			return (SubAssemblySection->SectionType == ESubAssemblySectionType::Reference) || (Sequence != nullptr) ? true : false;
+		}));
+
+
+	// Add the SubSequence property to the Sequence Category, but hide it when the SectionType is Template, because the Assembly Template property will be displayed instead
+	TSharedRef<IPropertyHandle> SubSequencePropertyHandle = DetailBuilder.GetProperty(FName(TEXT("SubSequence")), UMovieSceneSubSection::StaticClass());
+	IDetailPropertyRow& SubSequencePropertyRow = SequenceCategory.AddProperty(SubSequencePropertyHandle);
+	SubSequencePropertyRow.Visibility(TAttribute<EVisibility>::CreateLambda([SubAssemblySection]()
+		{ 
+			return (SubAssemblySection->SectionType != ESubAssemblySectionType::Template) ? EVisibility::Visible : EVisibility::Collapsed; 
+		}));
+
+	// Add the AssemblyTemplate property to the Sequence Category, but hide it when the SectionType is Reference, because the SubSequence property will be displayed instead
+	TSharedRef<IPropertyHandle> AssemblyTemplatePropertyHandle = DetailBuilder.GetProperty(UMovieSceneSubAssemblySection::AssemblyTemplatePropertyName);
+	IDetailPropertyRow& AssemblyTemplatePropertyRow = SequenceCategory.AddProperty(AssemblyTemplatePropertyHandle);
+	AssemblyTemplatePropertyRow.Visibility(TAttribute<EVisibility>::CreateLambda([SubAssemblySection]()
+		{
+			return (SubAssemblySection->SectionType == ESubAssemblySectionType::Template) ? EVisibility::Visible : EVisibility::Collapsed;
+		}));
+
+	// Further customize the AssemblyTemplate property so that the asset picker widget only allows Level Sequences, Cine Assemblies, and Schemas
+	AssemblyTemplatePropertyRow.CustomWidget()
+		.NameContent()
+		[
+			AssemblyTemplatePropertyHandle->CreatePropertyNameWidget()
+		]
+		.ValueContent()
+		[
+			SNew(SObjectPropertyEntryBox)
+				.AllowedClass(UObject::StaticClass())
+				.ThumbnailPool(DetailBuilder.GetThumbnailPool())
+				.AllowCreate(true)
+				.OnShouldFilterAsset_Raw(this, &FSubAssemblySectionDetailCustomization::OnShouldFilterAsset)
+				.ObjectPath_Lambda([SubAssemblySection]()
+					{			
+						return (SubAssemblySection->GetAssemblyTemplate()) ? SubAssemblySection->GetAssemblyTemplate()->GetPathName() : TEXT("");
+					})
+				.OnObjectChanged_Lambda([SubAssemblySection](const FAssetData& InAssetData)
+					{
+						SubAssemblySection->SetAssemblyTemplate(InAssetData.GetAsset());
+					})
+		];
+}
+
+bool FSubAssemblySectionDetailCustomization::OnShouldFilterAsset(const FAssetData& InAssetData)
+{
+	// If none of these classes match, filter out the asset
+	return (InAssetData.GetClass() != ULevelSequence::StaticClass()) 
+		&& (InAssetData.GetClass() != UCineAssembly::StaticClass()) 
+		&& (InAssetData.GetClass() != UCineAssemblySchema::StaticClass());
+}
\ No newline at end of file
diff --git a/Engine/Plugins/VirtualProduction/CinematicAssemblyTools/Source/CineAssemblyToolsEditor/Private/Sequencer/MovieSceneSubAssemblySectionCustomization.h b/Engine/Plugins/VirtualProduction/CinematicAssemblyTools/Source/CineAssemblyToolsEditor/Private/Sequencer/MovieSceneSubAssemblySectionCustomization.h
new file mode 100644
--- /dev/null
+++ b/Engine/Plugins/VirtualProduction/CinematicAssemblyTools/Source/CineAssemblyToolsEditor/Private/Sequencer/MovieSceneSubAssemblySectionCustomization.h
@@ -0,0 +1,23 @@
+// Copyright Epic Games, Inc. All Rights Reserved.
+
+#pragma once
+
+#include "IDetailCustomization.h"
+
+struct FAssetData;
+
+/** Detail Customization for UMovieSceneSubAssemblySection */
+class FSubAssemblySectionDetailCustomization : public IDetailCustomization
+{
+public:
+	/** Makes a new instance of this customization for a specific detail view requesting it */
+	static TSharedRef<IDetailCustomization> MakeInstance();
+
+	//~ Begin IDetailCustomization interface */
+	virtual void CustomizeDetails(IDetailLayoutBuilder& DetailBuilder) override;
+	//~ End IDetailCustomization interface */
+
+private:
+	/** Determines whether the input asset should be filtered out of the Assembly Template asset picker */
+	bool OnShouldFilterAsset(const FAssetData& InAssetData);
+};
diff --git a/Engine/Plugins/VirtualProduction/CinematicAssemblyTools/Source/CineAssemblyToolsEditor/Private/Sequencer/SubAssemblySection.cpp b/Engine/Plugins/VirtualProduction/CinematicAssemblyTools/Source/CineAssemblyToolsEditor/Private/Sequencer/SubAssemblySection.cpp
--- a/Engine/Plugins/VirtualProduction/CinematicAssemblyTools/Source/CineAssemblyToolsEditor/Private/Sequencer/SubAssemblySection.cpp
+++ b/Engine/Plugins/VirtualProduction/CinematicAssemblyTools/Source/CineAssemblyToolsEditor/Private/Sequencer/SubAssemblySection.cpp
@@ -3,30 +3,20 @@
 #include "SubAssemblySection.h"
 
 #include "CineAssembly.h"
+#include "IDetailsView.h"
+#include "MovieSceneSubAssemblySectionCustomization.h"
 
 #define LOCTEXT_NAMESPACE "CineAssemblySchemaSubSection"
 
-FSubAssemblySection::FSubAssemblySection(TSharedPtr<ISequencer> InSequencer, UMovieSceneSubSection& InSectionObject)
-	: TSubSectionMixin(InSequencer, *CastChecked<UMovieSceneSubSection>(&InSectionObject))
-	, SectionObject(InSectionObject)
+FSubAssemblySection::FSubAssemblySection(TSharedPtr<ISequencer> InSequencer, UMovieSceneSubAssemblySection* InSection)
+	: TSubSectionMixin(InSequencer, *CastChecked<UMovieSceneSubSection>(InSection))
+	, MovieSceneSubAssemblySection(InSection)
 {
-	// If the Section's Sequence is owned by a Schema, then this section represents a template for a new sequence, not a reference to an existing one
-	if (UMovieSceneSequence* Sequence = SectionObject.GetSequence())
-	{
-		if (Sequence->GetTypedOuter<UCineAssemblySchema>())
-		{
-			bIsTemplateSection = true;
-		}
-		else
-		{
-			bIsTemplateSection = false;
-		}
-	}
 }
 
 UMovieSceneSection* FSubAssemblySection::GetSectionObject()
 {
-	return &SectionObject;
+	return MovieSceneSubAssemblySection;
 }
 
 FText FSubAssemblySection::GetSectionTitle() const
@@ -49,7 +39,6 @@ TSharedRef<SWidget> FSubAssemblySection::GenerateSectionWidget()
 				.ShouldEvaluateTokens(false)
 				.DisplayTokenIcon(false)
 				.DisplayBorderImage(false)
-				.IsReadOnly(!bIsTemplateSection)
 				.Text(this, &FSubAssemblySection::GetSectionText)
 				.OnTextCommitted(this, &FSubAssemblySection::SetSectionText)
 		]
@@ -64,27 +53,12 @@ TSharedRef<SWidget> FSubAssemblySection::GenerateSectionWidget()
 		];
 }
 
-int32 FSubAssemblySection::OnPaintSection(FSequencerSectionPainter& InPainter) const
-{
-	TSubSectionMixin<>::OnPaintSection(InPainter);
-
-	// Add a specific tint to Template Sections to more easily distinguish them from Reference Sections
-	if (bIsTemplateSection)
-	{
-		const FColor CineAssemblyAssetColor = FColor(229, 45, 113);
-		return InPainter.PaintSectionBackground(CineAssemblyAssetColor);
-	}
-
-	return InPainter.PaintSectionBackground();
-}
-
 FReply FSubAssemblySection::OnSectionDoubleClicked(const FGeometry& SectionGeometry, const FPointerEvent& MouseEvent)
 {
 	// TODO: Check the geometry to determine what double clicking should do.
 	// If double clicking on the Naming Token Textbox, it should allow the user to rename the section.
 	// If double clicking anywhere else, and the Section's sequence is a "Schema Template", then it should open that Schema in a separate window.
-
-	if (bIsTemplateSection && (SectionObject.GetSequence() != nullptr))
+	if (MovieSceneSubAssemblySection->IsTemplateSection())
 	{
 		EditableTextBlock->EnterEditingMode();
 	}
@@ -94,45 +68,29 @@ FReply FSubAssemblySection::OnSectionDoubleClicked(const FGeometry& SectionGeome
 
 FText FSubAssemblySection::GetSectionText() const
 {
-	UMovieSceneSequence* Sequence = SectionObject.GetSequence();
-	if (!Sequence)
+	if (MovieSceneSubAssemblySection->IsTemplateSection())
 	{
-		return FText::GetEmpty();
-	}
-
-	if (bIsTemplateSection)
-	{
-		// Template sections will display the (potentially tokenized) name of the assembly template
-		if (UCineAssembly* CineAssembly = Cast<UCineAssembly>(Sequence))
-		{
-			return FText::FromString(CineAssembly->AssemblyName.Template);
-		}
+		return MovieSceneSubAssemblySection->NewSequenceName;
 	}
 
-	// Reference sections simply display the name of the referenced sequence
-	return Sequence->GetDisplayName();
+	return MovieSceneSubAssemblySection->GetSequence() ? MovieSceneSubAssemblySection->GetSequence()->GetDisplayName() : FText::GetEmpty();
 }
 
 void FSubAssemblySection::SetSectionText(const FText& InText, ETextCommit::Type InCommitType)
 {
-	if (UCineAssembly* CineAssembly = Cast<UCineAssembly>(SectionObject.GetSequence()))
+	if (MovieSceneSubAssemblySection->IsTemplateSection())
 	{
-		CineAssembly->AssemblyName.Template = InText.ToString();
+		MovieSceneSubAssemblySection->NewSequenceName = InText;
 	}
 }
 
 FText FSubAssemblySection::GetOriginText() const
 {
-	if (bIsTemplateSection)
+	if (MovieSceneSubAssemblySection->IsTemplateSection())
 	{
-		// The origin text of a template section indicates the name of the original sequence that was used as the template
-		if (UCineAssembly* CineAssembly = Cast<UCineAssembly>(SectionObject.GetSequence()))
+		if (UObject* TemplateObject = MovieSceneSubAssemblySection->GetAssemblyTemplate())
 		{
-			FText TemplateOriginText = CineAssembly->AssemblyName.Resolved;
-			if (!TemplateOriginText.IsEmpty())
-			{
-				return FText::Format(LOCTEXT("TemplateOriginText", "from {0}"), TemplateOriginText);
-			}
+			return FText::Format(LOCTEXT("TemplateOriginText", "from {0}"), FText::FromString(TemplateObject->GetName()));
 		}
 		return FText::GetEmpty();
 	}
@@ -141,4 +99,24 @@ FText FSubAssemblySection::GetOriginText() const
 	return LOCTEXT("ReferenceText", "Reference");
 }
 
-#undef LOCTEXT_NAMESPACE
+void FSubAssemblySection::CustomizePropertiesDetailsView(TSharedRef<IDetailsView> DetailsView, const FSequencerSectionPropertyDetailsViewCustomizationParams& InParams) const
+{
+	DetailsView->RegisterInstancedCustomPropertyLayout(UMovieSceneSubAssemblySection::StaticClass(), FOnGetDetailCustomizationInstance::CreateStatic(&FSubAssemblySectionDetailCustomization::MakeInstance));
+}
+
+void FSubAssemblySection::BuildSectionContextMenu(FMenuBuilder& MenuBuilder, const FGuid& ObjectBinding)
+{
+	MenuBuilder.AddMenuEntry(
+		LOCTEXT("RenameAction", "Rename"),
+		LOCTEXT("RenameActionTooltip", "Rename the new sequence that will be created from this section"),
+		FSlateIcon(),
+		FUIAction(
+			FExecuteAction::CreateLambda([this]() { EditableTextBlock->EnterEditingMode(); }),
+			FCanExecuteAction::CreateLambda([this]() { return MovieSceneSubAssemblySection->IsTemplateSection(); })
+		),
+		NAME_None,
+		EUserInterfaceActionType::None
+	);
+}
+
+#undef LOCTEXT_NAMESPACE
\ No newline at end of file
diff --git a/Engine/Plugins/VirtualProduction/CinematicAssemblyTools/Source/CineAssemblyToolsEditor/Private/Sequencer/SubAssemblySection.h b/Engine/Plugins/VirtualProduction/CinematicAssemblyTools/Source/CineAssemblyToolsEditor/Private/Sequencer/SubAssemblySection.h
--- a/Engine/Plugins/VirtualProduction/CinematicAssemblyTools/Source/CineAssemblyToolsEditor/Private/Sequencer/SubAssemblySection.h
+++ b/Engine/Plugins/VirtualProduction/CinematicAssemblyTools/Source/CineAssemblyToolsEditor/Private/Sequencer/SubAssemblySection.h
@@ -4,44 +4,42 @@
 
 #include "ISequencerSection.h"
 
-#include "Sections/MovieSceneSubSection.h"
+#include "MovieSceneSubAssemblySection.h"
 #include "SNamingTokensEditableTextBox.h"
 #include "TrackEditors/SubTrackEditorBase.h"
 
 /**
- * An ISequencerSection that handles the UI of sections within a SubAssembly Track
+ * An ISequencerSection that handles the UI of a UMovieSceneSubAssemblySection
  */
-class FSubAssemblySection	: public TSubSectionMixin<>, public TSharedFromThis<FSubAssemblySection>
+class FSubAssemblySection : public TSubSectionMixin<>, public TSharedFromThis<FSubAssemblySection>
 {
 public:
-	FSubAssemblySection(TSharedPtr<ISequencer> InSequencer, UMovieSceneSubSection& InSectionObject);
+	FSubAssemblySection(TSharedPtr<ISequencer> InSequencer, UMovieSceneSubAssemblySection* InSection);
 
 public:
 	//~ Begin ISequencerSection interface
 	virtual UMovieSceneSection* GetSectionObject() override;
 	virtual FText GetSectionTitle() const override;
 	virtual TSharedRef<SWidget> GenerateSectionWidget() override;
-	virtual int32 OnPaintSection(FSequencerSectionPainter& InPainter) const override;
 	virtual FReply OnSectionDoubleClicked(const FGeometry& SectionGeometry, const FPointerEvent& MouseEvent) override;
+	virtual void BuildSectionContextMenu(FMenuBuilder& MenuBuilder, const FGuid& ObjectBinding) override;
+	virtual void CustomizePropertiesDetailsView(TSharedRef<IDetailsView> DetailsView, const FSequencerSectionPropertyDetailsViewCustomizationParams& InParams) const override;
 	//~ End ISequencerSection interface
 
 private:
-	/** Get the name of the subsequence in the Section to display in the Section Widget */
+	/** Get the name of the sequence to display in the Section Widget */
 	FText GetSectionText() const;
 
-	/** Set the name of the subsequence in the Section */
+	/** Set the name of the sequence in a template section */
 	void SetSectionText(const FText& InText, ETextCommit::Type InCommitType);
 
-	/** If this is a template section, return the name of the sequence used as a template. Otherwise, return text indicating that this is a reference section. */
+	/** If this is a template section, return the name of the template object. Otherwise, return text indicating that this is a reference section. */
 	FText GetOriginText() const;
 
 private:
-	/** Section object represented by this SubSection */
-	UMovieSceneSubSection& SectionObject;
+	/** MovieSceneSubAssemblySection handled by this sequencer section */
+	UMovieSceneSubAssemblySection* MovieSceneSubAssemblySection;
 
-	/** True if the sequence was added to this section as a Template (and not a Reference) */
-	bool bIsTemplateSection = true;
-
-	/** Editable textbox widget, used for renaming the Sequence in this Section */
+	/** Editable textbox widget, used for renaming the sequence in template sections */
 	TSharedPtr<SNamingTokensEditableTextBox> EditableTextBlock;
 };
diff --git a/Engine/Plugins/VirtualProduction/CinematicAssemblyTools/Source/CineAssemblyToolsEditor/Private/Sequencer/SubAssemblyTrackEditor.cpp b/Engine/Plugins/VirtualProduction/CinematicAssemblyTools/Source/CineAssemblyToolsEditor/Private/Sequencer/SubAssemblyTrackEditor.cpp
--- a/Engine/Plugins/VirtualProduction/CinematicAssemblyTools/Source/CineAssemblyToolsEditor/Private/Sequencer/SubAssemblyTrackEditor.cpp
+++ b/Engine/Plugins/VirtualProduction/CinematicAssemblyTools/Source/CineAssemblyToolsEditor/Private/Sequencer/SubAssemblyTrackEditor.cpp
@@ -6,6 +6,7 @@
 #include "CineAssemblySchema.h"
 #include "ContentBrowserModule.h"
 #include "IContentBrowserSingleton.h"
+#include "MovieSceneSubAssemblySection.h"
 #include "MovieSceneSubAssemblyTrack.h"
 #include "MovieSceneToolsProjectSettings.h"
 #include "MVVM/Extensions/ITrackExtension.h"
@@ -55,8 +56,8 @@ bool FSubAssemblyTrackEditor::SupportsSequence(UMovieSceneSequence* InSequence)
 
 TSharedRef<ISequencerSection> FSubAssemblyTrackEditor::MakeSectionInterface(UMovieSceneSection& SectionObject, UMovieSceneTrack& Track, FGuid ObjectBinding)
 {
-	UMovieSceneSubSection* SubSection = Cast<UMovieSceneSubSection>(&SectionObject);
-	return MakeShared<FSubAssemblySection>(GetSequencer(), *SubSection);
+	UMovieSceneSubAssemblySection* SubAssemblySection = Cast<UMovieSceneSubAssemblySection>(&SectionObject);
+	return MakeShared<FSubAssemblySection>(GetSequencer(), SubAssemblySection);
 }
 
 TSharedPtr<SWidget> FSubAssemblyTrackEditor::BuildOutlinerEditWidget(const FGuid& ObjectBinding, UMovieSceneTrack* Track, const FBuildEditWidgetParams& Params)
@@ -146,7 +147,7 @@ void FSubAssemblyTrackEditor::AddNewSequenceSubMenu(FMenuBuilder& MenuBuilder, U
 		LOCTEXT("EmptySequence", "Empty Sequence"),
 		LOCTEXT("EmptySequenceTooltip", "Add a new blank Sequence to this track"),
 		FSlateIcon(),
-		FUIAction(FExecuteAction::CreateRaw(this, &FSubAssemblyTrackEditor::AddEmptySequence, Track))
+		FUIAction(FExecuteAction::CreateRaw(this, &FSubAssemblyTrackEditor::AddNewSequenceToTrack, Track, (UObject*)nullptr))
 	);
 }
 
@@ -201,16 +202,29 @@ TSharedRef<SWidget> FSubAssemblyTrackEditor::BuildAssetPicker(UMovieSceneSubAsse
 		];
 }
 
-void FSubAssemblyTrackEditor::AddEmptySequence(UMovieSceneSubAssemblyTrack* Track)
+void FSubAssemblyTrackEditor::AddNewSequenceToTrack(UMovieSceneSubAssemblyTrack* Track, UObject* TemplateObject)
 {
-	UCineAssembly* NewAssembly = CreateCineAssemblyFromTemplate(nullptr);
+	// Create a new empty Assembly and add it to the that can be added to the SubAssemblyTrack
+	UCineAssembly* NewAssembly = NewObject<UCineAssembly>();
+	NewAssembly->Initialize();
+
 	FSubTrackEditor::AddSequenceToTrack(NewAssembly, Track);
+
+	// Find the SubAssemblySection in the Track that was just created (i.e. whose Sequence matches the one just created) and set its Assembly Template to be the input template object
+	for (UMovieSceneSection* Section : Track->GetAllSections())
+	{
+		UMovieSceneSubAssemblySection* SubAssemblySection = Cast<UMovieSceneSubAssemblySection>(Section);
+		if (SubAssemblySection->GetSequence() == NewAssembly)
+		{
+			SubAssemblySection->SetAssemblyTemplate(TemplateObject);
+			break;
+		}
+	}
 }
 
 void FSubAssemblyTrackEditor::OnAssetSelectedAsTemplate(const FAssetData& AssetData, UMovieSceneSubAssemblyTrack* Track)
 {
-	UCineAssembly* NewAssembly = CreateCineAssemblyFromTemplate(AssetData.GetAsset());
-	FSubTrackEditor::AddSequenceToTrack(NewAssembly, Track);
+	AddNewSequenceToTrack(Track, AssetData.GetAsset());
 }
 
 void FSubAssemblyTrackEditor::OnAssetEnterPressedAsTemplate(const TArray<FAssetData>& AssetData, UMovieSceneSubAssemblyTrack* Track)
@@ -237,41 +251,4 @@ void FSubAssemblyTrackEditor::OnAssetEnterPressedAsReference(const TArray<FAsset
 	}
 }
 
-UCineAssembly* FSubAssemblyTrackEditor::CreateCineAssemblyFromTemplate(UObject* TemplateObject)
-{
-	// Get the Schema Template Sequence being edited by the Schema Sequencer, as well as a pointer to the Schema itself
-	UMovieSceneSequence* SchemaTemplateSequence = GetSequencer()->GetFocusedMovieSceneSequence();
-	UCineAssemblySchema* Schema = SchemaTemplateSequence->GetTypedOuter<UCineAssemblySchema>();
-
-	// Create a new Assembly, owned by the Schema, and parented to the Schema Template Sequence
-	UCineAssembly* NewAssembly = NewObject<UCineAssembly>(Schema, NAME_None);
-	NewAssembly->SetParentAssembly(Schema->TemplateSequence);
-
-	if (UCineAssemblySchema* SchemaTemplate = Cast<UCineAssemblySchema>(TemplateObject))
-	{
-		NewAssembly->Initialize();
-		NewAssembly->SetSchema(SchemaTemplate);
-		NewAssembly->AssemblyName.Resolved = FText::FromString(SchemaTemplate->SchemaName);
-	}
-	else if (UCineAssembly* AssemblyTemplate = Cast<UCineAssembly>(TemplateObject))
-	{
-		NewAssembly->InitializeFromTemplate(AssemblyTemplate);
-		NewAssembly->AssemblyName = AssemblyTemplate->AssemblyName;
-	}
-	else if (ULevelSequence* SequenceTemplate = Cast<ULevelSequence>(TemplateObject))
-	{
-		NewAssembly->InitializeFromTemplate(SequenceTemplate);
-		NewAssembly->AssemblyName.Template = SequenceTemplate->GetDisplayName().ToString();
-		NewAssembly->AssemblyName.Resolved = FText::GetEmpty();
-	}
-	else
-	{
-		NewAssembly->Initialize();
-		NewAssembly->AssemblyName.Template = TEXT("NewAssembly");
-		NewAssembly->AssemblyName.Resolved = FText::GetEmpty();
-	}
-
-	return NewAssembly;
-}
-
 #undef LOCTEXT_NAMESPACE
diff --git a/Engine/Plugins/VirtualProduction/CinematicAssemblyTools/Source/CineAssemblyToolsEditor/Private/Sequencer/SubAssemblyTrackEditor.h b/Engine/Plugins/VirtualProduction/CinematicAssemblyTools/Source/CineAssemblyToolsEditor/Private/Sequencer/SubAssemblyTrackEditor.h
--- a/Engine/Plugins/VirtualProduction/CinematicAssemblyTools/Source/CineAssemblyToolsEditor/Private/Sequencer/SubAssemblyTrackEditor.h
+++ b/Engine/Plugins/VirtualProduction/CinematicAssemblyTools/Source/CineAssemblyToolsEditor/Private/Sequencer/SubAssemblyTrackEditor.h
@@ -54,9 +54,6 @@ private:
 	/** Callback when one or more assets are selected in the asset picker as a reference to add to the track */
 	void OnAssetEnterPressedAsReference(const TArray<FAssetData>& AssetData, UMovieSceneSubAssemblyTrack* Track);
 
-	/** Creates a new empty Assembly (no schema) and adds it to the Track */
-	void AddEmptySequence(UMovieSceneSubAssemblyTrack* Track);
-
-	/** Creates a new CineAssembly and initializes it using the input Template object, which can be null */
-	UCineAssembly* CreateCineAssemblyFromTemplate(UObject* TemplateObject);
+	/** Creates and adds a new sequence to the input Track, then finds the newly created SubAssemblySection and assigns the input TemplateObject to it */
+	void AddNewSequenceToTrack(UMovieSceneSubAssemblyTrack* Track, UObject* TemplateObject);
 };
\ No newline at end of file
