[Backout] - CL49894259
#rnx
Original CL Desc
-----------------------------------------------------------------
ShotManagement: Add new MovieSceneSubAssemblySection, which derives from MovieSceneSubSection, so that the section itself can manage the template asset used to create new SubAssemblies.
Add the ability to change a SubAssemblySection between a Reference and a Template, and to change the Template object via the details panel.
The SubAssemblyTrackEditor now no longer uses the template when adding a new section to the track. Instead it simply creates an empty assembly.

#Synced-CL 49044222

#Synced-CL 49044222

#Synced-CL 49044222

diff --git a/Engine/Plugins/VirtualProduction/CinematicAssemblyTools/Source/CineAssemblyTools/Private/CineAssembly.cpp b/Engine/Plugins/VirtualProduction/CinematicAssemblyTools/Source/CineAssemblyTools/Private/CineAssembly.cpp
--- a/Engine/Plugins/VirtualProduction/CinematicAssemblyTools/Source/CineAssemblyTools/Private/CineAssembly.cpp
+++ b/Engine/Plugins/VirtualProduction/CinematicAssemblyTools/Source/CineAssemblyTools/Private/CineAssembly.cpp
@@ -9,7 +9,6 @@
 #include "LevelSequenceShotMetaDataLibrary.h"
 #include "MovieScene.h"
 #include "MovieSceneFolder.h"
-#include "MovieSceneSubAssemblySection.h"
 #include "MovieSceneSubAssemblyTrack.h"
 #include "NamingTokensEngineSubsystem.h"
 #include "Sections/MovieSceneSubSection.h"
@@ -224,20 +223,25 @@ void UCineAssembly::InitializeFromTemplate(ULevelSequence* Template)
 
 			for (UMovieSceneSection* Section : SubAssemblyTrack->GetAllSections())
 			{
-				if (UMovieSceneSubAssemblySection* SubAssemblySection = Cast<UMovieSceneSubAssemblySection>(Section))
+				if (UMovieSceneSubSection* SchemaSubSection = Cast<UMovieSceneSubSection>(Section))
 				{
-					// Create a new SubSection, with the relevant properties copied from the SubAssemblySection, and add it to the new SubTrack
-					UMovieSceneSubSection* NewSubSection = NewObject<UMovieSceneSubSection>(SubTrack, NAME_None, SubAssemblySection->GetFlags(), SubAssemblySection);
-					SubTrack->AddSection(*NewSubSection);
-
-					if (SubAssemblySection->IsTemplateSection())
+					// Duplicate the SubSection from the SubAssemblyTrack and add it to the new normal SubTrack
+					UMovieSceneSubSection* DuplicateSection = DuplicateObject(SchemaSubSection, SubTrack);
+					SubTrack->AddSection(*DuplicateSection);
+
+					// Check if the SubSequence in this SubSection should be added as a new sequence or as a reference, by checking its owner.
+					// A new SubSequence will be owned by the Schema. A reference would be refer to an asset that is owned by its UPackage. 
+					UMovieSceneSequence* SubSequence = SchemaSubSection->GetSequence();
+					if (SubSequence && SubSequence->GetTypedOuter<UCineAssemblySchema>())
 					{
-						// Create a new SubAssembly using the assembly template specified by the SubAssemblySection
-						UCineAssembly* NewSubAssembly = CreateSubAssemblyFromTemplate(SubAssemblySection->NewSequenceName.ToString(), SubAssemblySection->GetAssemblyTemplate());
+						// Create a new SubAssembly asset and set it as the SubSection's sequence (replacing the assembly template from the Schema).
+						UCineAssembly* AssemblyTemplate = Cast<UCineAssembly>(SubSequence);
+						UCineAssembly* NewSubAssembly = CreateSubAssemblyFromTemplate(AssemblyTemplate);
+
 						if (NewSubAssembly)
 						{
-							NewSubSection->SetSequence(NewSubAssembly);
-							SubAssemblies.Add(NewSubSection);
+							DuplicateSection->SetSequence(NewSubAssembly);
+							SubAssemblies.Add(DuplicateSection);
 						}
 					}
 				}
@@ -290,43 +294,25 @@ void UCineAssembly::InitializeFromTemplate(ULevelSequence* Template)
 }
 
 #if WITH_EDITOR
-UCineAssembly* UCineAssembly::CreateSubAssemblyFromTemplate(const FString& SubAssemblyName, UObject* TemplateObject)
+UCineAssembly* UCineAssembly::CreateSubAssemblyFromTemplate(UCineAssembly* AssemblyTemplate)
 {
-	// The new SubAssembly begins are a transient object so that it can be configured before it is named and given a package.
-	// This allows us to set its metadata first before trying to resolve any tokens which might appear in the SubAssemblyName.
-	UCineAssembly* SubAssembly = NewObject<UCineAssembly>(GetTransientPackage(), NAME_None, RF_Public | RF_Standalone | RF_Transactional);
-
-	// Initialize the new SubAssembly using the TemplateObject
-	if (UCineAssemblySchema* SchemaTemplate = Cast<UCineAssemblySchema>(TemplateObject))
-	{
-		SubAssembly->Initialize();
-		SubAssembly->SetSchema(SchemaTemplate);
-	}
-	else if (ULevelSequence* SequenceTemplate = Cast<ULevelSequence>(TemplateObject))
+	if (!AssemblyTemplate)
 	{
-		SubAssembly->InitializeFromTemplate(SequenceTemplate);
-	}
-	else
-	{
-		SubAssembly->Initialize();
+		return nullptr;
 	}
 
-	SubAssembly->Level = Level;
-	SubAssembly->ParentAssembly = this;
-	SubAssembly->Production = Production;
-	SubAssembly->ProductionName = ProductionName;
+	// Resolve any tokens found in the name of the SubAssembly name before attempting to create it 
+	FTemplateString SubAssemblyName = AssemblyTemplate->AssemblyName;
+	SubAssemblyName.Resolved = UCineAssemblyNamingTokens::GetResolvedText(SubAssemblyName.Template, AssemblyTemplate);
 
-	// Evaluate the name of the assembly
-	SubAssembly->AssemblyName.Template = SubAssemblyName;
-	SubAssembly->AssemblyName.Resolved = UCineAssemblyNamingTokens::GetResolvedText(SubAssembly->AssemblyName.Template, SubAssembly);
+	const FString SubAssemblyFilename = FPaths::GetBaseFilename(SubAssemblyName.Resolved.ToString());
 
-	// If the assembly name is empty or invalid, assign it a valid default name
-	if (SubAssembly->AssemblyName.Resolved.IsEmpty() || FName(*SubAssembly->AssemblyName.Resolved.ToString()).IsNone())
+	if (SubAssemblyFilename.IsEmpty())
 	{
-		SubAssembly->AssemblyName.Resolved = LOCTEXT("NewCineAssemblyName", "NewCineAssembly");
-		SubAssembly->AssemblyName.Template = SubAssembly->AssemblyName.Resolved.ToString();
+		return nullptr;
 	}
 
+	// Before creating each subassembly, sanity check that each one will actually have a unique name (in case there are duplicates in the schema description)
 	FString AssemblyPath;
 	FString AssemblyRootFolder;
 	GetAssetPathAndRootFolder(AssemblyPath, AssemblyRootFolder);
@@ -334,19 +320,34 @@ UCineAssembly* UCineAssembly::CreateSubAssemblyFromTemplate(const FString& SubAs
 	IAssetTools& AssetTools = FModuleManager::GetModuleChecked<FAssetToolsModule>("AssetTools").Get();
 	FString UniquePackageName;
 	FString UniqueAssetName;
-	AssetTools.CreateUniqueAssetName(AssemblyRootFolder / SubAssembly->AssemblyName.Resolved.ToString(), TEXT(""), UniquePackageName, UniqueAssetName);
+	AssetTools.CreateUniqueAssetName(AssemblyRootFolder / SubAssemblyName.Resolved.ToString(), TEXT(""), UniquePackageName, UniqueAssetName);
+
+	// Create a new Package and CineAssembly object for the new SubAssembly
+	UPackage* SubAssemblyPackage = CreatePackage(*UniquePackageName);
 
-	// Now, we can create a real package for it and rename it
-	UPackage* Package = CreatePackage(*UniquePackageName);
-	SubAssembly->Rename(*UniqueAssetName, Package);
+	const EObjectFlags Flags = RF_Public | RF_Standalone | RF_Transactional;
+	UCineAssembly* SubAssembly = NewObject<UCineAssembly>(SubAssemblyPackage, UCineAssembly::StaticClass(), FName(*UniqueAssetName), Flags);
 
-	SubAssembly->CreateSubAssemblies();
+	if (SubAssembly)
+	{
+		SubAssembly->AssemblyName = SubAssemblyName;
+		SubAssembly->InitializeFromTemplate(AssemblyTemplate);
+
+		SubAssembly->Level = Level;
+		SubAssembly->ParentAssembly = this;
+		SubAssembly->Production = Production;
+		SubAssembly->ProductionName = ProductionName;
+
+		SubAssembly->CreateSubAssemblies();
 
-	// Notify the asset registry about the new assembly
-	FAssetRegistryModule::AssetCreated(SubAssembly);
+		ULevelSequenceShotMetaDataLibrary::SetIsSubSequence(SubAssembly, true);
 
-	// Mark the package dirty
-	Package->MarkPackageDirty();
+		// Notify the asset registry about the new subassembly
+		FAssetRegistryModule::AssetCreated(SubAssembly);
+
+		// Mark the package dirty
+		SubAssemblyPackage->MarkPackageDirty();
+	}
 
 	return SubAssembly;
 }
@@ -396,7 +397,14 @@ void UCineAssembly::ChangeSchema(UCineAssemblySchema* InSchema)
 	// Reset the assembly's name based on the schema template
 	if (BaseSchema)
 	{
-		AssemblyName.Template = BaseSchema->GetDefaultAssemblyName();
+		if (!BaseSchema->DefaultAssemblyName.IsEmpty())
+		{
+			AssemblyName.Template = BaseSchema->DefaultAssemblyName;
+		}
+		else
+		{
+			AssemblyName.Template = TEXT("New") + BaseSchema->SchemaName;
+		}
 	}
 	else
 	{
diff --git a/Engine/Plugins/VirtualProduction/CinematicAssemblyTools/Source/CineAssemblyTools/Private/CineAssemblySchema.cpp b/Engine/Plugins/VirtualProduction/CinematicAssemblyTools/Source/CineAssemblyTools/Private/CineAssemblySchema.cpp
--- a/Engine/Plugins/VirtualProduction/CinematicAssemblyTools/Source/CineAssemblyTools/Private/CineAssemblySchema.cpp
+++ b/Engine/Plugins/VirtualProduction/CinematicAssemblyTools/Source/CineAssemblyTools/Private/CineAssemblySchema.cpp
@@ -71,15 +71,6 @@ void UCineAssemblySchema::PostLoad()
 	}
 }
 
-FString UCineAssemblySchema::GetDefaultAssemblyName() const
-{
-	if (DefaultAssemblyName.IsEmpty())
-	{
-		return TEXT("New") + SchemaName;
-	}
-	return DefaultAssemblyName;
-}
-
 bool UCineAssemblySchema::SupportsRename() const
 {
 	return bSupportsRename;
diff --git a/Engine/Plugins/VirtualProduction/CinematicAssemblyTools/Source/CineAssemblyTools/Private/MovieSceneSubAssemblySection.cpp b/Engine/Plugins/VirtualProduction/CinematicAssemblyTools/Source/CineAssemblyTools/Private/MovieSceneSubAssemblySection.cpp
deleted file mode 100644
diff --git a/Engine/Plugins/VirtualProduction/CinematicAssemblyTools/Source/CineAssemblyTools/Private/MovieSceneSubAssemblyTrack.cpp b/Engine/Plugins/VirtualProduction/CinematicAssemblyTools/Source/CineAssemblyTools/Private/MovieSceneSubAssemblyTrack.cpp
--- a/Engine/Plugins/VirtualProduction/CinematicAssemblyTools/Source/CineAssemblyTools/Private/MovieSceneSubAssemblyTrack.cpp
+++ b/Engine/Plugins/VirtualProduction/CinematicAssemblyTools/Source/CineAssemblyTools/Private/MovieSceneSubAssemblyTrack.cpp
@@ -2,14 +2,7 @@
 
 #include "MovieSceneSubAssemblyTrack.h"
 
-#include "MovieSceneSubAssemblySection.h"
-
 UMovieSceneSubAssemblyTrack::UMovieSceneSubAssemblyTrack(const FObjectInitializer& ObjectInitializer)
 	: Super(ObjectInitializer)
 {
 }
-
-UMovieSceneSection* UMovieSceneSubAssemblyTrack::CreateNewSection()
-{
-	return NewObject<UMovieSceneSubAssemblySection>(this, NAME_None, RF_Transactional);
-}
diff --git a/Engine/Plugins/VirtualProduction/CinematicAssemblyTools/Source/CineAssemblyTools/Public/CineAssembly.h b/Engine/Plugins/VirtualProduction/CinematicAssemblyTools/Source/CineAssemblyTools/Public/CineAssembly.h
--- a/Engine/Plugins/VirtualProduction/CinematicAssemblyTools/Source/CineAssemblyTools/Public/CineAssembly.h
+++ b/Engine/Plugins/VirtualProduction/CinematicAssemblyTools/Source/CineAssemblyTools/Public/CineAssembly.h
@@ -191,8 +191,8 @@ private:
 	TOptional<FDateTime> TryGetCreatedAsLocalTime() const;
 
 #if WITH_EDITOR
-	/** Creates a new CineAssembly asset using the input object as a template */
-	UCineAssembly* CreateSubAssemblyFromTemplate(const FString& SubAssemblyName, UObject* TemplateObject);
+	/** Creates a new CineAssembly asset using the input assembly as a template */
+	UCineAssembly* CreateSubAssemblyFromTemplate(UCineAssembly* AssemblyTemplate);
 
 	/** Recursively gather the movie scene folders within a root folder */
 	void GetAllMovieSceneFoldersRecursive(UMovieSceneFolder* RootFolder, TArray<UMovieSceneFolder*>& OutFolders);
diff --git a/Engine/Plugins/VirtualProduction/CinematicAssemblyTools/Source/CineAssemblyTools/Public/CineAssemblySchema.h b/Engine/Plugins/VirtualProduction/CinematicAssemblyTools/Source/CineAssemblyTools/Public/CineAssemblySchema.h
--- a/Engine/Plugins/VirtualProduction/CinematicAssemblyTools/Source/CineAssemblyTools/Public/CineAssemblySchema.h
+++ b/Engine/Plugins/VirtualProduction/CinematicAssemblyTools/Source/CineAssemblyTools/Public/CineAssemblySchema.h
@@ -73,9 +73,6 @@ public:
 	/** Get the unique ID of this schema */
 	CINEASSEMBLYTOOLS_API FGuid GetSchemaGuid() const;
 
-	/** Get the default name for a new Assembly made from this schema */
-	CINEASSEMBLYTOOLS_API FString GetDefaultAssemblyName() const;
-
 	/** Whether or not this schema can be renamed */
 	CINEASSEMBLYTOOLS_API bool SupportsRename() const;
 
diff --git a/Engine/Plugins/VirtualProduction/CinematicAssemblyTools/Source/CineAssemblyTools/Public/MovieSceneSubAssemblySection.h b/Engine/Plugins/VirtualProduction/CinematicAssemblyTools/Source/CineAssemblyTools/Public/MovieSceneSubAssemblySection.h
deleted file mode 100644
diff --git a/Engine/Plugins/VirtualProduction/CinematicAssemblyTools/Source/CineAssemblyTools/Public/MovieSceneSubAssemblyTrack.h b/Engine/Plugins/VirtualProduction/CinematicAssemblyTools/Source/CineAssemblyTools/Public/MovieSceneSubAssemblyTrack.h
--- a/Engine/Plugins/VirtualProduction/CinematicAssemblyTools/Source/CineAssemblyTools/Public/MovieSceneSubAssemblyTrack.h
+++ b/Engine/Plugins/VirtualProduction/CinematicAssemblyTools/Source/CineAssemblyTools/Public/MovieSceneSubAssemblyTrack.h
@@ -6,23 +6,15 @@
 
 #include "MovieSceneSubAssemblyTrack.generated.h"
 
-#define UE_API CINEASSEMBLYTOOLS_API
-
 /**
  * A Track to hold SubAssemblies within a parent CineAssembly.
  * Used by the Schema Sequencer tool to support custom behavior above and beyond normal SubTracks
  */
-UCLASS(BlueprintType, MinimalAPI)
+UCLASS()
 class UMovieSceneSubAssemblyTrack : public UMovieSceneSubTrack
 {
 	GENERATED_BODY()
 
 public:
 	UMovieSceneSubAssemblyTrack(const FObjectInitializer& ObjectInitializer);
-
-	//~ Begin UMovieSceneTrack interface
-	UE_API virtual UMovieSceneSection* CreateNewSection() override;
-	//~ End UMovieSceneTrack interface
 };
-
-#undef UE_API
diff --git a/Engine/Plugins/VirtualProduction/CinematicAssemblyTools/Source/CineAssemblyToolsEditor/Private/Sequencer/MovieSceneSubAssemblySectionCustomization.cpp b/Engine/Plugins/VirtualProduction/CinematicAssemblyTools/Source/CineAssemblyToolsEditor/Private/Sequencer/MovieSceneSubAssemblySectionCustomization.cpp
deleted file mode 100644
diff --git a/Engine/Plugins/VirtualProduction/CinematicAssemblyTools/Source/CineAssemblyToolsEditor/Private/Sequencer/MovieSceneSubAssemblySectionCustomization.h b/Engine/Plugins/VirtualProduction/CinematicAssemblyTools/Source/CineAssemblyToolsEditor/Private/Sequencer/MovieSceneSubAssemblySectionCustomization.h
deleted file mode 100644
diff --git a/Engine/Plugins/VirtualProduction/CinematicAssemblyTools/Source/CineAssemblyToolsEditor/Private/Sequencer/SubAssemblySection.cpp b/Engine/Plugins/VirtualProduction/CinematicAssemblyTools/Source/CineAssemblyToolsEditor/Private/Sequencer/SubAssemblySection.cpp
--- a/Engine/Plugins/VirtualProduction/CinematicAssemblyTools/Source/CineAssemblyToolsEditor/Private/Sequencer/SubAssemblySection.cpp
+++ b/Engine/Plugins/VirtualProduction/CinematicAssemblyTools/Source/CineAssemblyToolsEditor/Private/Sequencer/SubAssemblySection.cpp
@@ -3,20 +3,30 @@
 #include "SubAssemblySection.h"
 
 #include "CineAssembly.h"
-#include "IDetailsView.h"
-#include "MovieSceneSubAssemblySectionCustomization.h"
 
 #define LOCTEXT_NAMESPACE "CineAssemblySchemaSubSection"
 
-FSubAssemblySection::FSubAssemblySection(TSharedPtr<ISequencer> InSequencer, UMovieSceneSubAssemblySection* InSection)
-	: TSubSectionMixin(InSequencer, *CastChecked<UMovieSceneSubSection>(InSection))
-	, MovieSceneSubAssemblySection(InSection)
+FSubAssemblySection::FSubAssemblySection(TSharedPtr<ISequencer> InSequencer, UMovieSceneSubSection& InSectionObject)
+	: TSubSectionMixin(InSequencer, *CastChecked<UMovieSceneSubSection>(&InSectionObject))
+	, SectionObject(InSectionObject)
 {
+	// If the Section's Sequence is owned by a Schema, then this section represents a template for a new sequence, not a reference to an existing one
+	if (UMovieSceneSequence* Sequence = SectionObject.GetSequence())
+	{
+		if (Sequence->GetTypedOuter<UCineAssemblySchema>())
+		{
+			bIsTemplateSection = true;
+		}
+		else
+		{
+			bIsTemplateSection = false;
+		}
+	}
 }
 
 UMovieSceneSection* FSubAssemblySection::GetSectionObject()
 {
-	return MovieSceneSubAssemblySection;
+	return &SectionObject;
 }
 
 FText FSubAssemblySection::GetSectionTitle() const
@@ -39,6 +49,7 @@ TSharedRef<SWidget> FSubAssemblySection::GenerateSectionWidget()
 				.ShouldEvaluateTokens(false)
 				.DisplayTokenIcon(false)
 				.DisplayBorderImage(false)
+				.IsReadOnly(!bIsTemplateSection)
 				.Text(this, &FSubAssemblySection::GetSectionText)
 				.OnTextCommitted(this, &FSubAssemblySection::SetSectionText)
 		]
@@ -53,12 +64,27 @@ TSharedRef<SWidget> FSubAssemblySection::GenerateSectionWidget()
 		];
 }
 
+int32 FSubAssemblySection::OnPaintSection(FSequencerSectionPainter& InPainter) const
+{
+	TSubSectionMixin<>::OnPaintSection(InPainter);
+
+	// Add a specific tint to Template Sections to more easily distinguish them from Reference Sections
+	if (bIsTemplateSection)
+	{
+		const FColor CineAssemblyAssetColor = FColor(229, 45, 113);
+		return InPainter.PaintSectionBackground(CineAssemblyAssetColor);
+	}
+
+	return InPainter.PaintSectionBackground();
+}
+
 FReply FSubAssemblySection::OnSectionDoubleClicked(const FGeometry& SectionGeometry, const FPointerEvent& MouseEvent)
 {
 	// TODO: Check the geometry to determine what double clicking should do.
 	// If double clicking on the Naming Token Textbox, it should allow the user to rename the section.
 	// If double clicking anywhere else, and the Section's sequence is a "Schema Template", then it should open that Schema in a separate window.
-	if (MovieSceneSubAssemblySection->IsTemplateSection())
+
+	if (bIsTemplateSection && (SectionObject.GetSequence() != nullptr))
 	{
 		EditableTextBlock->EnterEditingMode();
 	}
@@ -68,29 +94,45 @@ FReply FSubAssemblySection::OnSectionDoubleClicked(const FGeometry& SectionGeome
 
 FText FSubAssemblySection::GetSectionText() const
 {
-	if (MovieSceneSubAssemblySection->IsTemplateSection())
+	UMovieSceneSequence* Sequence = SectionObject.GetSequence();
+	if (!Sequence)
 	{
-		return MovieSceneSubAssemblySection->NewSequenceName;
+		return FText::GetEmpty();
+	}
+
+	if (bIsTemplateSection)
+	{
+		// Template sections will display the (potentially tokenized) name of the assembly template
+		if (UCineAssembly* CineAssembly = Cast<UCineAssembly>(Sequence))
+		{
+			return FText::FromString(CineAssembly->AssemblyName.Template);
+		}
 	}
 
-	return MovieSceneSubAssemblySection->GetSequence() ? MovieSceneSubAssemblySection->GetSequence()->GetDisplayName() : FText::GetEmpty();
+	// Reference sections simply display the name of the referenced sequence
+	return Sequence->GetDisplayName();
 }
 
 void FSubAssemblySection::SetSectionText(const FText& InText, ETextCommit::Type InCommitType)
 {
-	if (MovieSceneSubAssemblySection->IsTemplateSection())
+	if (UCineAssembly* CineAssembly = Cast<UCineAssembly>(SectionObject.GetSequence()))
 	{
-		MovieSceneSubAssemblySection->NewSequenceName = InText;
+		CineAssembly->AssemblyName.Template = InText.ToString();
 	}
 }
 
 FText FSubAssemblySection::GetOriginText() const
 {
-	if (MovieSceneSubAssemblySection->IsTemplateSection())
+	if (bIsTemplateSection)
 	{
-		if (UObject* TemplateObject = MovieSceneSubAssemblySection->GetAssemblyTemplate())
+		// The origin text of a template section indicates the name of the original sequence that was used as the template
+		if (UCineAssembly* CineAssembly = Cast<UCineAssembly>(SectionObject.GetSequence()))
 		{
-			return FText::Format(LOCTEXT("TemplateOriginText", "from {0}"), FText::FromString(TemplateObject->GetName()));
+			FText TemplateOriginText = CineAssembly->AssemblyName.Resolved;
+			if (!TemplateOriginText.IsEmpty())
+			{
+				return FText::Format(LOCTEXT("TemplateOriginText", "from {0}"), TemplateOriginText);
+			}
 		}
 		return FText::GetEmpty();
 	}
@@ -99,24 +141,4 @@ FText FSubAssemblySection::GetOriginText() const
 	return LOCTEXT("ReferenceText", "Reference");
 }
 
-void FSubAssemblySection::CustomizePropertiesDetailsView(TSharedRef<IDetailsView> DetailsView, const FSequencerSectionPropertyDetailsViewCustomizationParams& InParams) const
-{
-	DetailsView->RegisterInstancedCustomPropertyLayout(UMovieSceneSubAssemblySection::StaticClass(), FOnGetDetailCustomizationInstance::CreateStatic(&FSubAssemblySectionDetailCustomization::MakeInstance));
-}
-
-void FSubAssemblySection::BuildSectionContextMenu(FMenuBuilder& MenuBuilder, const FGuid& ObjectBinding)
-{
-	MenuBuilder.AddMenuEntry(
-		LOCTEXT("RenameAction", "Rename"),
-		LOCTEXT("RenameActionTooltip", "Rename the new sequence that will be created from this section"),
-		FSlateIcon(),
-		FUIAction(
-			FExecuteAction::CreateLambda([this]() { EditableTextBlock->EnterEditingMode(); }),
-			FCanExecuteAction::CreateLambda([this]() { return MovieSceneSubAssemblySection->IsTemplateSection(); })
-		),
-		NAME_None,
-		EUserInterfaceActionType::None
-	);
-}
-
-#undef LOCTEXT_NAMESPACE
\ No newline at end of file
+#undef LOCTEXT_NAMESPACE
diff --git a/Engine/Plugins/VirtualProduction/CinematicAssemblyTools/Source/CineAssemblyToolsEditor/Private/Sequencer/SubAssemblySection.h b/Engine/Plugins/VirtualProduction/CinematicAssemblyTools/Source/CineAssemblyToolsEditor/Private/Sequencer/SubAssemblySection.h
--- a/Engine/Plugins/VirtualProduction/CinematicAssemblyTools/Source/CineAssemblyToolsEditor/Private/Sequencer/SubAssemblySection.h
+++ b/Engine/Plugins/VirtualProduction/CinematicAssemblyTools/Source/CineAssemblyToolsEditor/Private/Sequencer/SubAssemblySection.h
@@ -4,42 +4,44 @@
 
 #include "ISequencerSection.h"
 
-#include "MovieSceneSubAssemblySection.h"
+#include "Sections/MovieSceneSubSection.h"
 #include "SNamingTokensEditableTextBox.h"
 #include "TrackEditors/SubTrackEditorBase.h"
 
 /**
- * An ISequencerSection that handles the UI of a UMovieSceneSubAssemblySection
+ * An ISequencerSection that handles the UI of sections within a SubAssembly Track
  */
-class FSubAssemblySection : public TSubSectionMixin<>, public TSharedFromThis<FSubAssemblySection>
+class FSubAssemblySection	: public TSubSectionMixin<>, public TSharedFromThis<FSubAssemblySection>
 {
 public:
-	FSubAssemblySection(TSharedPtr<ISequencer> InSequencer, UMovieSceneSubAssemblySection* InSection);
+	FSubAssemblySection(TSharedPtr<ISequencer> InSequencer, UMovieSceneSubSection& InSectionObject);
 
 public:
 	//~ Begin ISequencerSection interface
 	virtual UMovieSceneSection* GetSectionObject() override;
 	virtual FText GetSectionTitle() const override;
 	virtual TSharedRef<SWidget> GenerateSectionWidget() override;
+	virtual int32 OnPaintSection(FSequencerSectionPainter& InPainter) const override;
 	virtual FReply OnSectionDoubleClicked(const FGeometry& SectionGeometry, const FPointerEvent& MouseEvent) override;
-	virtual void BuildSectionContextMenu(FMenuBuilder& MenuBuilder, const FGuid& ObjectBinding) override;
-	virtual void CustomizePropertiesDetailsView(TSharedRef<IDetailsView> DetailsView, const FSequencerSectionPropertyDetailsViewCustomizationParams& InParams) const override;
 	//~ End ISequencerSection interface
 
 private:
-	/** Get the name of the sequence to display in the Section Widget */
+	/** Get the name of the subsequence in the Section to display in the Section Widget */
 	FText GetSectionText() const;
 
-	/** Set the name of the sequence in a template section */
+	/** Set the name of the subsequence in the Section */
 	void SetSectionText(const FText& InText, ETextCommit::Type InCommitType);
 
-	/** If this is a template section, return the name of the template object. Otherwise, return text indicating that this is a reference section. */
+	/** If this is a template section, return the name of the sequence used as a template. Otherwise, return text indicating that this is a reference section. */
 	FText GetOriginText() const;
 
 private:
-	/** MovieSceneSubAssemblySection handled by this sequencer section */
-	UMovieSceneSubAssemblySection* MovieSceneSubAssemblySection;
+	/** Section object represented by this SubSection */
+	UMovieSceneSubSection& SectionObject;
 
-	/** Editable textbox widget, used for renaming the sequence in template sections */
+	/** True if the sequence was added to this section as a Template (and not a Reference) */
+	bool bIsTemplateSection = true;
+
+	/** Editable textbox widget, used for renaming the Sequence in this Section */
 	TSharedPtr<SNamingTokensEditableTextBox> EditableTextBlock;
 };
diff --git a/Engine/Plugins/VirtualProduction/CinematicAssemblyTools/Source/CineAssemblyToolsEditor/Private/Sequencer/SubAssemblyTrackEditor.cpp b/Engine/Plugins/VirtualProduction/CinematicAssemblyTools/Source/CineAssemblyToolsEditor/Private/Sequencer/SubAssemblyTrackEditor.cpp
--- a/Engine/Plugins/VirtualProduction/CinematicAssemblyTools/Source/CineAssemblyToolsEditor/Private/Sequencer/SubAssemblyTrackEditor.cpp
+++ b/Engine/Plugins/VirtualProduction/CinematicAssemblyTools/Source/CineAssemblyToolsEditor/Private/Sequencer/SubAssemblyTrackEditor.cpp
@@ -6,7 +6,6 @@
 #include "CineAssemblySchema.h"
 #include "ContentBrowserModule.h"
 #include "IContentBrowserSingleton.h"
-#include "MovieSceneSubAssemblySection.h"
 #include "MovieSceneSubAssemblyTrack.h"
 #include "MovieSceneToolsProjectSettings.h"
 #include "MVVM/Extensions/ITrackExtension.h"
@@ -56,8 +55,8 @@ bool FSubAssemblyTrackEditor::SupportsSequence(UMovieSceneSequence* InSequence)
 
 TSharedRef<ISequencerSection> FSubAssemblyTrackEditor::MakeSectionInterface(UMovieSceneSection& SectionObject, UMovieSceneTrack& Track, FGuid ObjectBinding)
 {
-	UMovieSceneSubAssemblySection* SubAssemblySection = Cast<UMovieSceneSubAssemblySection>(&SectionObject);
-	return MakeShared<FSubAssemblySection>(GetSequencer(), SubAssemblySection);
+	UMovieSceneSubSection* SubSection = Cast<UMovieSceneSubSection>(&SectionObject);
+	return MakeShared<FSubAssemblySection>(GetSequencer(), *SubSection);
 }
 
 TSharedPtr<SWidget> FSubAssemblyTrackEditor::BuildOutlinerEditWidget(const FGuid& ObjectBinding, UMovieSceneTrack* Track, const FBuildEditWidgetParams& Params)
@@ -147,7 +146,7 @@ void FSubAssemblyTrackEditor::AddNewSequenceSubMenu(FMenuBuilder& MenuBuilder, U
 		LOCTEXT("EmptySequence", "Empty Sequence"),
 		LOCTEXT("EmptySequenceTooltip", "Add a new blank Sequence to this track"),
 		FSlateIcon(),
-		FUIAction(FExecuteAction::CreateRaw(this, &FSubAssemblyTrackEditor::AddNewSequenceToTrack, Track, (UObject*)nullptr))
+		FUIAction(FExecuteAction::CreateRaw(this, &FSubAssemblyTrackEditor::AddEmptySequence, Track))
 	);
 }
 
@@ -202,29 +201,16 @@ TSharedRef<SWidget> FSubAssemblyTrackEditor::BuildAssetPicker(UMovieSceneSubAsse
 		];
 }
 
-void FSubAssemblyTrackEditor::AddNewSequenceToTrack(UMovieSceneSubAssemblyTrack* Track, UObject* TemplateObject)
+void FSubAssemblyTrackEditor::AddEmptySequence(UMovieSceneSubAssemblyTrack* Track)
 {
-	// Create a new empty Assembly and add it to the that can be added to the SubAssemblyTrack
-	UCineAssembly* NewAssembly = NewObject<UCineAssembly>();
-	NewAssembly->Initialize();
-
+	UCineAssembly* NewAssembly = CreateCineAssemblyFromTemplate(nullptr);
 	FSubTrackEditor::AddSequenceToTrack(NewAssembly, Track);
-
-	// Find the SubAssemblySection in the Track that was just created (i.e. whose Sequence matches the one just created) and set its Assembly Template to be the input template object
-	for (UMovieSceneSection* Section : Track->GetAllSections())
-	{
-		UMovieSceneSubAssemblySection* SubAssemblySection = Cast<UMovieSceneSubAssemblySection>(Section);
-		if (SubAssemblySection->GetSequence() == NewAssembly)
-		{
-			SubAssemblySection->SetAssemblyTemplate(TemplateObject);
-			break;
-		}
-	}
 }
 
 void FSubAssemblyTrackEditor::OnAssetSelectedAsTemplate(const FAssetData& AssetData, UMovieSceneSubAssemblyTrack* Track)
 {
-	AddNewSequenceToTrack(Track, AssetData.GetAsset());
+	UCineAssembly* NewAssembly = CreateCineAssemblyFromTemplate(AssetData.GetAsset());
+	FSubTrackEditor::AddSequenceToTrack(NewAssembly, Track);
 }
 
 void FSubAssemblyTrackEditor::OnAssetEnterPressedAsTemplate(const TArray<FAssetData>& AssetData, UMovieSceneSubAssemblyTrack* Track)
@@ -251,4 +237,41 @@ void FSubAssemblyTrackEditor::OnAssetEnterPressedAsReference(const TArray<FAsset
 	}
 }
 
+UCineAssembly* FSubAssemblyTrackEditor::CreateCineAssemblyFromTemplate(UObject* TemplateObject)
+{
+	// Get the Schema Template Sequence being edited by the Schema Sequencer, as well as a pointer to the Schema itself
+	UMovieSceneSequence* SchemaTemplateSequence = GetSequencer()->GetFocusedMovieSceneSequence();
+	UCineAssemblySchema* Schema = SchemaTemplateSequence->GetTypedOuter<UCineAssemblySchema>();
+
+	// Create a new Assembly, owned by the Schema, and parented to the Schema Template Sequence
+	UCineAssembly* NewAssembly = NewObject<UCineAssembly>(Schema, NAME_None);
+	NewAssembly->SetParentAssembly(Schema->TemplateSequence);
+
+	if (UCineAssemblySchema* SchemaTemplate = Cast<UCineAssemblySchema>(TemplateObject))
+	{
+		NewAssembly->Initialize();
+		NewAssembly->SetSchema(SchemaTemplate);
+		NewAssembly->AssemblyName.Resolved = FText::FromString(SchemaTemplate->SchemaName);
+	}
+	else if (UCineAssembly* AssemblyTemplate = Cast<UCineAssembly>(TemplateObject))
+	{
+		NewAssembly->InitializeFromTemplate(AssemblyTemplate);
+		NewAssembly->AssemblyName = AssemblyTemplate->AssemblyName;
+	}
+	else if (ULevelSequence* SequenceTemplate = Cast<ULevelSequence>(TemplateObject))
+	{
+		NewAssembly->InitializeFromTemplate(SequenceTemplate);
+		NewAssembly->AssemblyName.Template = SequenceTemplate->GetDisplayName().ToString();
+		NewAssembly->AssemblyName.Resolved = FText::GetEmpty();
+	}
+	else
+	{
+		NewAssembly->Initialize();
+		NewAssembly->AssemblyName.Template = TEXT("NewAssembly");
+		NewAssembly->AssemblyName.Resolved = FText::GetEmpty();
+	}
+
+	return NewAssembly;
+}
+
 #undef LOCTEXT_NAMESPACE
diff --git a/Engine/Plugins/VirtualProduction/CinematicAssemblyTools/Source/CineAssemblyToolsEditor/Private/Sequencer/SubAssemblyTrackEditor.h b/Engine/Plugins/VirtualProduction/CinematicAssemblyTools/Source/CineAssemblyToolsEditor/Private/Sequencer/SubAssemblyTrackEditor.h
--- a/Engine/Plugins/VirtualProduction/CinematicAssemblyTools/Source/CineAssemblyToolsEditor/Private/Sequencer/SubAssemblyTrackEditor.h
+++ b/Engine/Plugins/VirtualProduction/CinematicAssemblyTools/Source/CineAssemblyToolsEditor/Private/Sequencer/SubAssemblyTrackEditor.h
@@ -54,6 +54,9 @@ private:
 	/** Callback when one or more assets are selected in the asset picker as a reference to add to the track */
 	void OnAssetEnterPressedAsReference(const TArray<FAssetData>& AssetData, UMovieSceneSubAssemblyTrack* Track);
 
-	/** Creates and adds a new sequence to the input Track, then finds the newly created SubAssemblySection and assigns the input TemplateObject to it */
-	void AddNewSequenceToTrack(UMovieSceneSubAssemblyTrack* Track, UObject* TemplateObject);
+	/** Creates a new empty Assembly (no schema) and adds it to the Track */
+	void AddEmptySequence(UMovieSceneSubAssemblyTrack* Track);
+
+	/** Creates a new CineAssembly and initializes it using the input Template object, which can be null */
+	UCineAssembly* CreateCineAssemblyFromTemplate(UObject* TemplateObject);
 };
\ No newline at end of file
