ShotManagement: Add support for dropping Sequences, Assemblies, and Schemas directly onto the Schema Sequencer, letting the user choose between adding as a Template or as a Reference where possible.
#rnx

#Synced-CL 49044222

diff --git a/Engine/Plugins/VirtualProduction/CinematicAssemblyTools/Source/CineAssemblyToolsEditor/Private/Sequencer/SubAssemblyTrackEditor.cpp b/Engine/Plugins/VirtualProduction/CinematicAssemblyTools/Source/CineAssemblyToolsEditor/Private/Sequencer/SubAssemblyTrackEditor.cpp
--- a/Engine/Plugins/VirtualProduction/CinematicAssemblyTools/Source/CineAssemblyToolsEditor/Private/Sequencer/SubAssemblyTrackEditor.cpp
+++ b/Engine/Plugins/VirtualProduction/CinematicAssemblyTools/Source/CineAssemblyToolsEditor/Private/Sequencer/SubAssemblyTrackEditor.cpp
@@ -54,6 +54,61 @@ bool FSubAssemblyTrackEditor::SupportsSequence(UMovieSceneSequence* InSequence)
 	return false;
 }
 
+bool FSubAssemblyTrackEditor::HandleAssetAdded(UObject* Asset, const FGuid& TargetObjectGuid)
+{
+	// If the dropped asset is a Schema, it can only be added as a template
+	if (Asset->IsA(UCineAssemblySchema::StaticClass()))
+	{
+		OnDroppedAsTemplate(Asset);
+		return true;
+	}
+
+	// Otherwise, check that the dropped asset is a valid Sequence that is supported by this track editor
+	UMovieSceneSequence* Sequence = Cast<UMovieSceneSequence>(Asset);
+	if (!Sequence)
+	{
+		return false;
+	}
+
+	if (!SupportsSequence(Sequence))
+	{
+		return false;
+	}
+
+	if (!CanHandleAssetAdded(Sequence))
+	{
+		return false;
+	}
+
+	// Build and display a menu to let the user choose between adding the dropped sequence as a reference, or as a template for a new sequence
+	constexpr bool bShouldCloseAfterMenuSelection = true;
+	FMenuBuilder MenuBuilder(bShouldCloseAfterMenuSelection, nullptr);
+
+	MenuBuilder.AddMenuEntry(
+		LOCTEXT("OnDroppedNewSequence", "Add New Sequence"),
+		LOCTEXT("OnDroppedNewSequenceTooltip", "Add a new Subsequence using the dropped Sequence as the base"),
+		FSlateIcon(),
+		FUIAction(FExecuteAction::CreateRaw(this, &FSubAssemblyTrackEditor::OnDroppedAsTemplate, Asset))
+	);
+
+	MenuBuilder.AddMenuEntry(
+		LOCTEXT("OnDroppedReference", "Add Reference to Sequence"),
+		LOCTEXT("OnDroppedReferenceTooltip", "Add a reference to the dropped Sequence"),
+		FSlateIcon(),
+		FUIAction(FExecuteAction::CreateRaw(this, &FSubAssemblyTrackEditor::OnDroppedAsReference, Sequence))
+	);
+
+	// Open the menu where at the user's current mouse position
+	UE::Slate::FDeprecateVector2DParameter MousePosition = FSlateApplication::Get().GetCursorPos();
+	FWidgetPath WidgetPath = FSlateApplication::Get().LocateWindowUnderMouse(MousePosition, FSlateApplication::Get().GetInteractiveTopLevelWindows(), false, 0);
+	if (WidgetPath.IsValid() && WidgetPath.Widgets.Num() > 0)
+	{
+		FSlateApplication::Get().PushMenu(WidgetPath.GetLastWidget(), WidgetPath, MenuBuilder.MakeWidget(), MousePosition, FPopupTransitionEffect::ContextMenu);
+	}
+
+	return true;
+}
+
 TSharedRef<ISequencerSection> FSubAssemblyTrackEditor::MakeSectionInterface(UMovieSceneSection& SectionObject, UMovieSceneTrack& Track, FGuid ObjectBinding)
 {
 	UMovieSceneSubAssemblySection* SubAssemblySection = Cast<UMovieSceneSubAssemblySection>(&SectionObject);
@@ -251,4 +306,20 @@ void FSubAssemblyTrackEditor::OnAssetEnterPressedAsReference(const TArray<FAsset
 	}
 }
 
+void FSubAssemblyTrackEditor::OnDroppedAsTemplate(UObject* Asset)
+{
+	UMovieScene* FocusedMovieScene = GetFocusedMovieScene();
+	UMovieSceneSubAssemblyTrack* NewTrack = Cast<UMovieSceneSubAssemblyTrack>(FindOrCreateSubTrack(FocusedMovieScene, nullptr));
+
+	AddNewSequenceToTrack(NewTrack, Asset);
+}
+
+void FSubAssemblyTrackEditor::OnDroppedAsReference(UMovieSceneSequence* Sequence)
+{
+	UMovieScene* FocusedMovieScene = GetFocusedMovieScene();
+	UMovieSceneSubAssemblyTrack* NewTrack = Cast<UMovieSceneSubAssemblyTrack>(FindOrCreateSubTrack(FocusedMovieScene, nullptr));
+
+	FSubTrackEditor::AddSequenceToTrack(Sequence, NewTrack);
+}
+
 #undef LOCTEXT_NAMESPACE
diff --git a/Engine/Plugins/VirtualProduction/CinematicAssemblyTools/Source/CineAssemblyToolsEditor/Private/Sequencer/SubAssemblyTrackEditor.h b/Engine/Plugins/VirtualProduction/CinematicAssemblyTools/Source/CineAssemblyToolsEditor/Private/Sequencer/SubAssemblyTrackEditor.h
--- a/Engine/Plugins/VirtualProduction/CinematicAssemblyTools/Source/CineAssemblyToolsEditor/Private/Sequencer/SubAssemblyTrackEditor.h
+++ b/Engine/Plugins/VirtualProduction/CinematicAssemblyTools/Source/CineAssemblyToolsEditor/Private/Sequencer/SubAssemblyTrackEditor.h
@@ -25,6 +25,7 @@ public:
 	//~ Begin ISequencerTrackEditor interface
 	virtual TSubclassOf<UMovieSceneSubTrack> GetSubTrackClass() const override;
 	virtual bool SupportsSequence(UMovieSceneSequence* InSequence) const override;
+	virtual bool HandleAssetAdded(UObject* Asset, const FGuid& TargetObjectGuid) override;
 	virtual TSharedPtr<SWidget> BuildOutlinerEditWidget(const FGuid& ObjectBinding, UMovieSceneTrack* Track, const FBuildEditWidgetParams& Params) override;
 	virtual TSharedRef<ISequencerSection> MakeSectionInterface(UMovieSceneSection& SectionObject, UMovieSceneTrack& Track, FGuid ObjectBinding) override;
 	//~ End ISequencerTrackEditor interface
@@ -56,4 +57,10 @@ private:
 
 	/** Creates and adds a new sequence to the input Track, then finds the newly created SubAssemblySection and assigns the input TemplateObject to it */
 	void AddNewSequenceToTrack(UMovieSceneSubAssemblyTrack* Track, UObject* TemplateObject);
+
+	/** Callback when an asset is dropped and handled by this track editor to be added as a template for a new sequence */
+	void OnDroppedAsTemplate(UObject* Asset);
+
+	/** Callback when an asset is dropped and handled by this track editor to be added as a reference */
+	void OnDroppedAsReference(UMovieSceneSequence* Sequence);
 };
\ No newline at end of file
