ShotManagement: Refactor the TreeView in the CIne Assembly Schema Content Hierarchy to use an interface for the tree items, which more concisely expresses the type of different items in the tree.
#rnx

#Synced-CL 49044222

diff --git a/Engine/Plugins/VirtualProduction/CinematicAssemblyTools/Source/CineAssemblyToolsEditor/Private/UI/CineAssembly/SCineAssemblySchemaWindow.cpp b/Engine/Plugins/VirtualProduction/CinematicAssemblyTools/Source/CineAssemblyToolsEditor/Private/UI/CineAssembly/SCineAssemblySchemaWindow.cpp
--- a/Engine/Plugins/VirtualProduction/CinematicAssemblyTools/Source/CineAssemblyToolsEditor/Private/UI/CineAssembly/SCineAssemblySchemaWindow.cpp
+++ b/Engine/Plugins/VirtualProduction/CinematicAssemblyTools/Source/CineAssemblyToolsEditor/Private/UI/CineAssembly/SCineAssemblySchemaWindow.cpp
@@ -41,9 +41,53 @@
 
 #define LOCTEXT_NAMESPACE "SCineAssemblySchemaWindow"
 
+struct ISchemaContentHierarchyTreeItem : TSharedFromThis<ISchemaContentHierarchyTreeItem>
+{
+	virtual ~ISchemaContentHierarchyTreeItem() {}
+
+	virtual TSharedPtr<FSchemaContentHierarchyFolder> AsFolder() { return nullptr; }
+	virtual TSharedPtr<FSchemaContentHierarchyAsset> AsAsset() { return nullptr; }
+
+	/** Whether this item can accept an item being dropped onto it */
+	virtual TOptional<EItemDropZone> CanAcceptDrop() const { return TOptional<EItemDropZone>(); }
+
+	/** The fully qualified path of this item (relative to the root folder) */
+	FString Path;
+
+	/** The parent folder of this item */
+	TSharedPtr<FSchemaContentHierarchyFolder> ParentFolder;
+
+	/** The text widget that displays this item's name and supports renaming */
+	TSharedPtr<SNamingTokensEditableTextBox> NameWidget;
+};
+
+struct FSchemaContentHierarchyFolder : ISchemaContentHierarchyTreeItem
+{
+	//~ Begin ISchemaContentHierarchyTreeItem Interface
+	virtual TSharedPtr<FSchemaContentHierarchyFolder> AsFolder() override { return SharedThis(this); }
+	virtual TOptional<EItemDropZone> CanAcceptDrop() const override { return EItemDropZone::OntoItem; }
+	//~ End ISchemaContentHierarchyTreeItem Interface
+
+	/** The child assets of this folder item */
+	TArray<TSharedPtr<FSchemaContentHierarchyAsset>> ChildAssets;
+
+	/** The child folders of this folder item */
+	TArray<TSharedPtr<FSchemaContentHierarchyFolder>> ChildFolders;
+};
+
+struct FSchemaContentHierarchyAsset : ISchemaContentHierarchyTreeItem
+{
+	//~ Begin ISchemaContentHierarchyTreeItem Interface
+	virtual TSharedPtr<FSchemaContentHierarchyAsset> AsAsset() override { return SharedThis(this); }
+	//~ End ISchemaContentHierarchyTreeItem Interface
+
+	/** Asset Template soft path for this asset item */
+	FSoftObjectPath AssetTemplate;
+};
+
 namespace UE::CineAssemblyTools::Private
 {
-	auto SortTreeItems = [](const TSharedPtr<FSchemaTreeItem>& A, const TSharedPtr<FSchemaTreeItem>& B) { return A->Path < B->Path; };
+	auto SortTreeItems = [](const TSharedPtr<ISchemaContentHierarchyTreeItem>& A, const TSharedPtr<ISchemaContentHierarchyTreeItem>& B) { return A->Path < B->Path; };
 };
 
 void SCineAssemblySchemaWindow::Construct(const FArguments& InArgs, const FString& InCreateAssetPath)
@@ -551,7 +595,7 @@ TSharedRef<SWidget> SCineAssemblySchemaWindow::MakeMetadataTabContent()
 
 TSharedRef<SWidget> SCineAssemblySchemaWindow::MakeHierarchyTabContent()
 {
-	TreeView = SNew(STreeView<TSharedPtr<FSchemaTreeItem>>)
+	TreeView = SNew(STreeView<TSharedPtr<ISchemaContentHierarchyTreeItem>>)
 		.TreeItemsSource(&TreeItems)
 		.SelectionMode(ESelectionMode::Single)
 		.OnGenerateRow(this, &SCineAssemblySchemaWindow::OnGenerateTreeRow)
@@ -561,7 +605,7 @@ TSharedRef<SWidget> SCineAssemblySchemaWindow::MakeHierarchyTabContent()
 		.OnMouseButtonDoubleClick(this, &SCineAssemblySchemaWindow::OnTreeViewDoubleClick)
 		.OnKeyDownHandler(this, &SCineAssemblySchemaWindow::OnTreeViewKeyDown);
 
-	ExpandTreeRecursive(RootItem);
+	ExpandTreeRecursive(RootFolder);
 
 	TSharedRef<SWidget> ContentHierarchyWidget = SNew(SVerticalBox)
 
@@ -711,10 +755,10 @@ FReply SCineAssemblySchemaWindow::OnAddAssetClicked()
 	return FReply::Handled();
 }
 
-FString SCineAssemblySchemaWindow::MakeUniqueFolderPath(TSharedPtr<FSchemaTreeItem> InItem)
+FString SCineAssemblySchemaWindow::MakeUniqueFolderPath(TSharedPtr<FSchemaContentHierarchyFolder> FolderItem)
 {
 	// This implementation is based on a similar utility in AssetTools for creating a unique asset name.
-	const FString BaseName = InItem->Path / TEXT("NewFolder");
+	const FString BaseName = FolderItem->Path / TEXT("NewFolder");
 
 	// Find the index in the string of the last non-numeric character
 	int32 CharIndex = BaseName.Len() - 1;
@@ -736,7 +780,7 @@ FString SCineAssemblySchemaWindow::MakeUniqueFolderPath(TSharedPtr<FSchemaTreeIt
 
 	FString WorkingName = TrimmedBaseName;
 
-	while (Algo::ContainsBy(InItem->ChildFolders, WorkingName, &FSchemaTreeItem::Path))
+	while (Algo::ContainsBy(FolderItem->ChildFolders, WorkingName, &FSchemaContentHierarchyFolder::Path))
 	{
 		WorkingName = FString::Printf(TEXT("%s%d"), *TrimmedBaseName, IntSuffix);
 		IntSuffix++;
@@ -747,34 +791,32 @@ FString SCineAssemblySchemaWindow::MakeUniqueFolderPath(TSharedPtr<FSchemaTreeIt
 
 FReply SCineAssemblySchemaWindow::OnAddFolderClicked()
 {
-	TSharedPtr<FSchemaTreeItem> NewFolder = MakeShared<FSchemaTreeItem>();
-	NewFolder->Type = FSchemaTreeItem::EItemType::Folder;
+	TSharedPtr<FSchemaContentHierarchyFolder> NewFolder = MakeShared<FSchemaContentHierarchyFolder>();
 
 	// Get the parent item for the new folder being added (this can be the root folder if no parent is currently selected)
 	// The TreeView uses single selection mode, so at most one item can ever be selected by the user
-	TSharedPtr<FSchemaTreeItem> SelectedTreeItem = RootItem;
+	TSharedPtr<FSchemaContentHierarchyFolder> ParentFolder = RootFolder;
 
-	TArray<TSharedPtr<FSchemaTreeItem>> SelectedNodes = TreeView->GetSelectedItems();
-	if (SelectedNodes.Num() == 1)
+	TArray<TSharedPtr<ISchemaContentHierarchyTreeItem>> SelectedItems = TreeView->GetSelectedItems();
+	if (SelectedItems.Num() == 1)
 	{
-		SelectedTreeItem = SelectedNodes[0];
+		ParentFolder = SelectedItems[0]->AsFolder() ? SelectedItems[0]->AsFolder() : SelectedItems[0]->ParentFolder;
 	}
 
-	const FString ParentPath = SelectedTreeItem->Path;
-	NewFolder->Path = MakeUniqueFolderPath(SelectedTreeItem);
-	NewFolder->Parent = SelectedTreeItem;
+	NewFolder->Path = MakeUniqueFolderPath(ParentFolder);
+	NewFolder->ParentFolder = ParentFolder;
 
-	SelectedTreeItem->ChildFolders.Add(NewFolder);
+	ParentFolder->ChildFolders.Add(NewFolder);
 
 	// Sort the children alphabetically to maintain a good ordering with the new folder
-	SelectedTreeItem->ChildFolders.Sort(UE::CineAssemblyTools::Private::SortTreeItems);
+	ParentFolder->ChildFolders.Sort(UE::CineAssemblyTools::Private::SortTreeItems);
 
 	// Save a reference to this item so that when the tree is rebuilt, we can immediately start editing its name
 	MostRecentlyAddedItem = NewFolder;
 
 	UpdateFolderList();
 
-	TreeView->SetItemExpansion(SelectedTreeItem, true);
+	TreeView->SetItemExpansion(ParentFolder, true);
 	TreeView->RequestTreeRefresh();
 
 	return FReply::Handled();
@@ -927,11 +969,11 @@ FReply SCineAssemblySchemaWindow::OnAssetListKeyDown(const FGeometry& MyGeometry
 	return FReply::Handled();
 }
 
-TSharedRef<ITableRow> SCineAssemblySchemaWindow::OnGenerateTreeRow(TSharedPtr<FSchemaTreeItem> TreeItem, const TSharedRef<STableViewBase>& OwnerTable)
+TSharedRef<ITableRow> SCineAssemblySchemaWindow::OnGenerateTreeRow(TSharedPtr<ISchemaContentHierarchyTreeItem> TreeItem, const TSharedRef<STableViewBase>& OwnerTable)
 {
 	TSharedPtr<SImage> Icon = SNew(SImage);
 
-	if (TreeItem->Type == FSchemaTreeItem::EItemType::Folder)
+	if (TreeItem->AsFolder())
 	{
 		Icon->SetImage(FCineAssemblyToolsStyle::Get().GetBrush("Icons.Folder"));
 		Icon->SetColorAndOpacity(FAppStyle::Get().GetSlateColor("ContentBrowser.DefaultFolderColor"));
@@ -942,16 +984,12 @@ TSharedRef<ITableRow> SCineAssemblySchemaWindow::OnGenerateTreeRow(TSharedPtr<FS
 		Icon->SetColorAndOpacity(FLinearColor::White);
 	}
 
-	const bool bIsRootItem = (TreeItem == RootItem);
+	const bool bIsRootItem = (TreeItem == RootFolder);
 
-	return SNew(STableRow<TSharedPtr<FSchemaTreeItem>>, OwnerTable)
+	return SNew(STableRow<TSharedPtr<ISchemaContentHierarchyTreeItem>>, OwnerTable)
 		.ShowSelection(true)
 		.Padding(FMargin(8.0f, 0.0f, 8.0f, 0.0f))
-		.OnCanAcceptDrop_Lambda([](const FDragDropEvent& InDragDropEvent, EItemDropZone InDropZone, TSharedPtr<FSchemaTreeItem> InItem)
-			{
-				// Only folder items can accept drops, and only onto to th item (not above or below)
-				return (InItem->Type == FSchemaTreeItem::EItemType::Folder) ? EItemDropZone::OntoItem : TOptional<EItemDropZone>();
-			})
+		.OnCanAcceptDrop_Lambda([](const FDragDropEvent& InDragDropEvent, EItemDropZone InDropZone, TSharedPtr<ISchemaContentHierarchyTreeItem> InItem) { return InItem->CanAcceptDrop(); })
 		.OnAcceptDrop(this, &SCineAssemblySchemaWindow::OnTreeRowAcceptDrop)
 		.OnDragDetected(this, &SCineAssemblySchemaWindow::OnTreeRowDragDetected)
 		.Content()
@@ -977,17 +1015,13 @@ TSharedRef<ITableRow> SCineAssemblySchemaWindow::OnGenerateTreeRow(TSharedPtr<FS
 						.DisplayTokenIcon(false)
 						.DisplayBorderImage(false)
 						.DisplayErrorMessage(!bIsRootItem)
-						.Text_Lambda([this, TreeItem]() -> FText
+						.Text_Lambda([bIsRootItem, TreeItem]() -> FText
 							{
-								return (TreeItem == RootItem) ? LOCTEXT("RootPathName", "Root Folder") : FText::FromString(FPaths::GetPathLeaf(TreeItem->Path));
+								return bIsRootItem ? LOCTEXT("RootPathName", "Root Folder") : FText::FromString(FPaths::GetPathLeaf(TreeItem->Path));
 							})
 						.OnValidateTokenizedText_Lambda([this, TreeItem](const FText& InText, FText& OutErrorMessage) -> bool
 							{
-								if (TreeItem->Type == FSchemaTreeItem::EItemType::Folder)
-								{
-									return ValidateFolderName(InText, OutErrorMessage, TreeItem);
-								}
-								return ValidateAssetName(InText, OutErrorMessage);
+								return TreeItem->AsFolder() ? ValidateFolderName(InText, OutErrorMessage, TreeItem->AsFolder()) : ValidateAssetName(InText, OutErrorMessage);
 							})
 						.OnTextCommitted(this, &SCineAssemblySchemaWindow::OnTreeItemTextCommitted, TreeItem)
 				]
@@ -1002,12 +1036,12 @@ TSharedRef<ITableRow> SCineAssemblySchemaWindow::OnGenerateTreeRow(TSharedPtr<FS
 						.DisplayCompactSize(true)
 						.AllowedClass(ULevelSequence::StaticClass())
 						.ToolTipText(LOCTEXT("SchemaAssetTemplateSelectionTooltip", "(Optional) Select an existing level sequence to use as a template for this subsequence."))
-						.Visibility(TreeItem->Type == FSchemaTreeItem::EItemType::Asset ? EVisibility::Visible : EVisibility::Collapsed)
-						.ObjectPath_Lambda([WeakTreeItem = TWeakPtr<FSchemaTreeItem>(TreeItem)]() -> FString
+						.Visibility(TreeItem->AsAsset() ? EVisibility::Visible : EVisibility::Collapsed)
+						.ObjectPath_Lambda([WeakTreeItem = TWeakPtr<ISchemaContentHierarchyTreeItem>(TreeItem)]() -> FString
 							{
-								if (TSharedPtr<FSchemaTreeItem> ItemPtr = WeakTreeItem.Pin())
+								if (TSharedPtr<ISchemaContentHierarchyTreeItem> ItemPtr = WeakTreeItem.Pin())
 								{
-									return ItemPtr->AssetTemplate.ToString();
+									return ItemPtr->AsAsset() ? ItemPtr->AsAsset()->AssetTemplate.ToString() : FString();
 								}
 								return FString();
 							})
@@ -1016,11 +1050,14 @@ TSharedRef<ITableRow> SCineAssemblySchemaWindow::OnGenerateTreeRow(TSharedPtr<FS
 		];
 }
 
-void SCineAssemblySchemaWindow::OnGetChildren(TSharedPtr<FSchemaTreeItem> TreeItem, TArray<TSharedPtr<FSchemaTreeItem>>& OutNodes)
+void SCineAssemblySchemaWindow::OnGetChildren(TSharedPtr<ISchemaContentHierarchyTreeItem> TreeItem, TArray<TSharedPtr<ISchemaContentHierarchyTreeItem>>& OutNodes)
 {
 	// Display all of the child assets first, followed by all of the child folders
-	OutNodes.Append(TreeItem->ChildAssets);
-	OutNodes.Append(TreeItem->ChildFolders);
+	if (TSharedPtr<FSchemaContentHierarchyFolder> FolderItem = TreeItem->AsFolder())
+	{
+		OutNodes.Append(FolderItem->ChildAssets);
+		OutNodes.Append(FolderItem->ChildFolders);
+	}
 }
 
 void SCineAssemblySchemaWindow::OnTreeItemsRebuilt()
@@ -1036,15 +1073,15 @@ void SCineAssemblySchemaWindow::OnTreeItemsRebuilt()
 TSharedPtr<SWidget> SCineAssemblySchemaWindow::MakeContentTreeContextMenu()
 {
 	// The TreeView uses single selection mode, so at most one item can ever be selected by the user
-	TArray<TSharedPtr<FSchemaTreeItem>> SelectedNodes = TreeView->GetSelectedItems();
-	if (SelectedNodes.Num() == 1)
+	TArray<TSharedPtr<ISchemaContentHierarchyTreeItem>> SelectedItems = TreeView->GetSelectedItems();
+	if (SelectedItems.Num() == 1)
 	{
 		constexpr bool bCloseAfterMenuSelection = true;
 		FMenuBuilder MenuBuilder(bCloseAfterMenuSelection, nullptr);
 
-		TSharedPtr<FSchemaTreeItem>& SelectedTreeItem = SelectedNodes[0];
+		TSharedPtr<ISchemaContentHierarchyTreeItem>& SelectedTreeItem = SelectedItems[0];
 
-		if (SelectedTreeItem->Type == FSchemaTreeItem::EItemType::Folder)
+		if (TSharedPtr<FSchemaContentHierarchyFolder> FolderItem = SelectedTreeItem->AsFolder())
 		{
 			MenuBuilder.AddMenuEntry(
 				LOCTEXT("AddChildFolderAction", "Add Child Folder"),
@@ -1056,17 +1093,13 @@ TSharedPtr<SWidget> SCineAssemblySchemaWindow::MakeContentTreeContextMenu()
 			);
 		}
 
-		if ((SelectedTreeItem != RootItem) && (SelectedTreeItem != TopLevelAssemblyItem)) 
+		if ((SelectedTreeItem != RootFolder) && (SelectedTreeItem != TopLevelAssemblyItem)) 
 		{
 			MenuBuilder.AddMenuEntry(
 				LOCTEXT("RenameAction", "Rename"),
 				LOCTEXT("RenameActionToolTip", "Rename"),
 				FSlateIcon(FCineAssemblyToolsStyle::StyleName, "Icons.AssetNaming"),
-				FUIAction(
-					FExecuteAction::CreateLambda([SelectedTreeItem]()
-						{
-							SelectedTreeItem->NameWidget->EnterEditingMode();
-						})),
+				FUIAction(FExecuteAction::CreateLambda([SelectedTreeItem]() { SelectedTreeItem->NameWidget->EnterEditingMode(); })),
 				NAME_None,
 				EUserInterfaceActionType::Button
 			);
@@ -1087,12 +1120,12 @@ TSharedPtr<SWidget> SCineAssemblySchemaWindow::MakeContentTreeContextMenu()
 	return SNullWidget::NullWidget;
 }
 
-void SCineAssemblySchemaWindow::OnTreeViewDoubleClick(TSharedPtr<FSchemaTreeItem> TreeItem)
+void SCineAssemblySchemaWindow::OnTreeViewDoubleClick(TSharedPtr<ISchemaContentHierarchyTreeItem> TreeItem)
 {
 	// Get the row for the input item and put its textblock in edit mode so the user can rename the item
 	if (TreeItem && TreeItem->NameWidget)
 	{
-		if (TreeItem != RootItem && TreeItem != TopLevelAssemblyItem)
+		if (TreeItem != RootFolder && TreeItem != TopLevelAssemblyItem)
 		{
 			TreeItem->NameWidget->EnterEditingMode();
 		}
@@ -1104,11 +1137,11 @@ FReply SCineAssemblySchemaWindow::OnTreeViewKeyDown(const FGeometry& MyGeometry,
 	if (InKeyEvent.GetKey() == EKeys::Delete)
 	{
 		// The TreeView uses single selection mode, so at most one item can ever be selected by the user
-		TArray<TSharedPtr<FSchemaTreeItem>> SelectedItems = TreeView->GetSelectedItems();
+		TArray<TSharedPtr<ISchemaContentHierarchyTreeItem>> SelectedItems = TreeView->GetSelectedItems();
 		if (SelectedItems.Num() == 1)
 		{
-			const TSharedPtr<FSchemaTreeItem>& SelectedItem = SelectedItems[0];
-			if ((SelectedItem != RootItem) && (SelectedItem != TopLevelAssemblyItem))
+			const TSharedPtr<ISchemaContentHierarchyTreeItem>& SelectedItem = SelectedItems[0];
+			if ((SelectedItem != RootFolder) && (SelectedItem != TopLevelAssemblyItem))
 			{
 				DeleteTreeItem(SelectedItem);
 			}
@@ -1121,121 +1154,124 @@ FReply SCineAssemblySchemaWindow::OnTreeViewKeyDown(const FGeometry& MyGeometry,
 void SCineAssemblySchemaWindow::InitializeContentTree()
 {
 	// Create the content tree root
-	RootItem = MakeShared<FSchemaTreeItem>();
-	RootItem->Type = FSchemaTreeItem::EItemType::Folder;
-	RootItem->Path = TEXT("");
+	RootFolder = MakeShared<FSchemaContentHierarchyFolder>();
+	RootFolder->Path = TEXT("");
+
+	TreeItems.Add(RootFolder);
 
-	TreeItems.Add(RootItem);
+	// Sort the list of folders so that paths are added to the tree in the proper order
+	SchemaToConfigure->FoldersToCreate.Sort();
 
-	auto AddItemsToTree = [this](TArray<FString>& ItemList, const TMap<FString, FSoftObjectPath>& Templates, FSchemaTreeItem::EItemType ItemType)
+	for (const FString& FolderPath : SchemaToConfigure->FoldersToCreate)
+	{
+		const FString ParentPath = FPaths::GetPath(FolderPath);
+
+		// Walk the tree until we find an item whose path matches the parent path. The new tree item will be created as one of its children
+		TSharedPtr<FSchemaContentHierarchyFolder> ParentFolder = FindFolderByPathRecursive(RootFolder, ParentPath);
+		if (ParentFolder)
 		{
-			// Sort the list so that paths are added to the tree in the proper order
-			ItemList.Sort();
+			TSharedPtr<FSchemaContentHierarchyFolder> NewItem = MakeShared<FSchemaContentHierarchyFolder>();
+			NewItem->Path = FolderPath;
+			NewItem->ParentFolder = ParentFolder;
 
-			for (const FString& ItemName : ItemList)
-			{
-				const FString ParentPath = FPaths::GetPath(ItemName);
+			ParentFolder->ChildFolders.Add(NewItem);
+		}
+	}
 
-				// Walk the tree until we find an item whose path matches the parent path. The new tree item will be created as one of its children
-				TSharedPtr<FSchemaTreeItem> ParentItem = FindItemAtPathRecursive(RootItem, ParentPath);
-				if (ParentItem)
-				{
-					TSharedPtr<FSchemaTreeItem> NewItem = MakeShared<FSchemaTreeItem>();
-					NewItem->Type = ItemType;
-					NewItem->Path = ItemName;
-					NewItem->Parent = ParentItem;
+	// Sort the list of subsequences so that paths are added to the tree in the proper order
+	SchemaToConfigure->SubsequencesToCreate.Sort();
 
-					if (const FSoftObjectPath* FoundTemplate = Templates.Find(ItemName))
-					{
-						NewItem->AssetTemplate = *FoundTemplate;
-					}
+	for (const FString& SubsequencePath : SchemaToConfigure->SubsequencesToCreate)
+	{
+		const FString ParentPath = FPaths::GetPath(SubsequencePath);
 
-					if (ItemType == FSchemaTreeItem::EItemType::Folder)
-					{
-						ParentItem->ChildFolders.Add(NewItem);
-					}
-					else
-					{
-						ParentItem->ChildAssets.Add(NewItem);
-					}
-				}
+		// Walk the tree until we find an item whose path matches the parent path. The new tree item will be created as one of its children
+		TSharedPtr<FSchemaContentHierarchyFolder> ParentFolder = FindFolderByPathRecursive(RootFolder, ParentPath);
+		if (ParentFolder)
+		{
+			TSharedPtr<FSchemaContentHierarchyAsset> NewItem = MakeShared<FSchemaContentHierarchyAsset>();
+			NewItem->Path = SubsequencePath;
+			NewItem->ParentFolder = ParentFolder;
+
+			if (const FSoftObjectPath* FoundTemplate = SchemaToConfigure->SubsequenceTemplates.Find(SubsequencePath))
+			{
+				NewItem->AssetTemplate = *FoundTemplate;
 			}
-		};
 
-	AddItemsToTree(SchemaToConfigure->FoldersToCreate, {}, FSchemaTreeItem::EItemType::Folder);
-	AddItemsToTree(SchemaToConfigure->SubsequencesToCreate, SchemaToConfigure->SubsequenceTemplates, FSchemaTreeItem::EItemType::Asset);
+			ParentFolder->ChildAssets.Add(NewItem);
+		}
+	}
 
 	// Add the top-level assembly node
-	TSharedPtr<FSchemaTreeItem> ParentItem = FindItemAtPathRecursive(RootItem, SchemaToConfigure->DefaultAssemblyPath);
-	if (ParentItem)
+	TSharedPtr<FSchemaContentHierarchyFolder> ParentFolder = FindFolderByPathRecursive(RootFolder, SchemaToConfigure->DefaultAssemblyPath);
+	if (ParentFolder)
 	{
-		TopLevelAssemblyItem = MakeShared<FSchemaTreeItem>();
-		TopLevelAssemblyItem->Type = FSchemaTreeItem::EItemType::Asset;
+		TopLevelAssemblyItem = MakeShared<FSchemaContentHierarchyAsset>();
 		TopLevelAssemblyItem->Path = TEXT("{assembly}");
 		TopLevelAssemblyItem->AssetTemplate = SchemaToConfigure->Template;
-		TopLevelAssemblyItem->Parent = ParentItem;
+		TopLevelAssemblyItem->ParentFolder = ParentFolder;
 
-		ParentItem->ChildAssets.Add(TopLevelAssemblyItem);
+		ParentFolder->ChildAssets.Add(TopLevelAssemblyItem);
 	}
 }
 
-void SCineAssemblySchemaWindow::ExpandTreeRecursive(TSharedPtr<FSchemaTreeItem> TreeItem) const
+void SCineAssemblySchemaWindow::ExpandTreeRecursive(TSharedPtr<FSchemaContentHierarchyFolder> FolderItem) const
 {
-	TreeView->SetItemExpansion(TreeItem, true);
+	TreeView->SetItemExpansion(FolderItem, true);
 
-	for (const TSharedPtr<FSchemaTreeItem>& ChildItem : TreeItem->ChildFolders)
+	for (const TSharedPtr<FSchemaContentHierarchyFolder>& ChildFolder : FolderItem->ChildFolders)
 	{
-		ExpandTreeRecursive(ChildItem);
+		ExpandTreeRecursive(ChildFolder);
 	}
 }
 
-void SCineAssemblySchemaWindow::GetFolderListRecursive(TSharedPtr<FSchemaTreeItem> TreeItem, TArray<FString>& FolderList) const
+void SCineAssemblySchemaWindow::GetFolderListRecursive(TSharedPtr<FSchemaContentHierarchyFolder> FolderItem, TArray<FString>& OutFolderList) const
 {
-	for (const TSharedPtr<FSchemaTreeItem>& Child : TreeItem->ChildFolders)
+	for (const TSharedPtr<FSchemaContentHierarchyFolder>& ChildFolder : FolderItem->ChildFolders)
 	{
-		FolderList.Add(Child->Path);
-		GetFolderListRecursive(Child, FolderList);
+		OutFolderList.Add(ChildFolder->Path);
+		GetFolderListRecursive(ChildFolder, OutFolderList);
 	}
 }
 
-void SCineAssemblySchemaWindow::GetAssetListRecursive(TSharedPtr<FSchemaTreeItem> TreeItem, TArray<FString>& AssetPathList, TMap<FString, FSoftObjectPath>& Templates) const
+void SCineAssemblySchemaWindow::GetAssetListRecursive(TSharedPtr<FSchemaContentHierarchyFolder> FolderItem, TArray<FString>& OutAssetList, TMap<FString, FSoftObjectPath>& OutTemplates) const
 {
-	for (const TSharedPtr<FSchemaTreeItem>& Asset : TreeItem->ChildAssets)
+	for (const TSharedPtr<FSchemaContentHierarchyAsset>& Asset : FolderItem->ChildAssets)
 	{
 		if (Asset != TopLevelAssemblyItem)
 		{
-			AssetPathList.Add(Asset->Path);
-			Templates.Add(Asset->Path, Asset->AssetTemplate);
+			OutAssetList.Add(Asset->Path);
+			OutTemplates.Add(Asset->Path, Asset->AssetTemplate);
 		}
 	}
 
-	for (const TSharedPtr<FSchemaTreeItem>& Child : TreeItem->ChildFolders)
+	for (const TSharedPtr<FSchemaContentHierarchyFolder>& ChildFolder : FolderItem->ChildFolders)
 	{
-		GetAssetListRecursive(Child, AssetPathList, Templates);
+		GetAssetListRecursive(ChildFolder, OutAssetList, OutTemplates);
 	}
 }
 
-TSharedPtr<FSchemaTreeItem> SCineAssemblySchemaWindow::FindItemAtPathRecursive(TSharedPtr<FSchemaTreeItem> TreeItem, const FString& Path) const
+TSharedPtr<FSchemaContentHierarchyFolder> SCineAssemblySchemaWindow::FindFolderByPathRecursive(TSharedPtr<FSchemaContentHierarchyFolder> FolderItem, const FString& FolderPath) const
 {
-	if (TreeItem->Path.Equals(Path))
+	if (FolderItem->Path.Equals(FolderPath))
 	{
-		return TreeItem;
+		return FolderItem;
 	}
 
-	for (const TSharedPtr<FSchemaTreeItem>& Child : TreeItem->ChildFolders)
+	for (const TSharedPtr<FSchemaContentHierarchyFolder>& ChildFolder : FolderItem->ChildFolders)
 	{
-		if (const TSharedPtr<FSchemaTreeItem>& ItemAtPath = FindItemAtPathRecursive(Child, Path))
+		if (const TSharedPtr<FSchemaContentHierarchyFolder>& MatchingFolder = FindFolderByPathRecursive(ChildFolder, FolderPath))
 		{
-			return ItemAtPath;
+			return MatchingFolder;
 		}
 	}
 
 	return nullptr;
 }
 
-void SCineAssemblySchemaWindow::SetChildrenPathRecursive(TSharedPtr<FSchemaTreeItem> TreeItem, const FString& NewPath)
+void SCineAssemblySchemaWindow::SetChildrenPathRecursive(TSharedPtr<FSchemaContentHierarchyFolder> FolderItem, const FString& NewPath)
 {
-	for (TSharedPtr<FSchemaTreeItem>& Asset : TreeItem->ChildAssets)
+	for (TSharedPtr<FSchemaContentHierarchyAsset>& Asset : FolderItem->ChildAssets)
 	{
 		const FString OldAssetName = FPaths::GetPathLeaf(Asset->Path);
 
@@ -1243,14 +1279,14 @@ void SCineAssemblySchemaWindow::SetChildrenPathRecursive(TSharedPtr<FSchemaTreeI
 		Asset->Path = NewAssetPath;
 	}
 
-	for (TSharedPtr<FSchemaTreeItem>& Child : TreeItem->ChildFolders)
+	for (TSharedPtr<FSchemaContentHierarchyFolder>& ChildFolder : FolderItem->ChildFolders)
 	{
-		const FString OldChildFolderName = FPaths::GetPathLeaf(Child->Path);
+		const FString OldChildFolderName = FPaths::GetPathLeaf(ChildFolder->Path);
 
 		const FString NewChildPath = NewPath / OldChildFolderName;
-		Child->Path = NewChildPath;
+		ChildFolder->Path = NewChildPath;
 
-		SetChildrenPathRecursive(Child, NewChildPath);
+		SetChildrenPathRecursive(ChildFolder, NewChildPath);
 	}
 }
 
@@ -1266,7 +1302,7 @@ void SCineAssemblySchemaWindow::EnterEditMode(TSharedPtr<FSchemaAssetParameters>
 	}
 }
 
-void SCineAssemblySchemaWindow::OnTreeItemTextCommitted(const FText& InText, ETextCommit::Type InCommitType, TSharedPtr<FSchemaTreeItem> TreeItem)
+void SCineAssemblySchemaWindow::OnTreeItemTextCommitted(const FText& InText, ETextCommit::Type InCommitType, TSharedPtr<ISchemaContentHierarchyTreeItem> TreeItem)
 {
 	// Early-out if the name has not actually changed
 	const FString OldPath = TreeItem->Path;
@@ -1279,16 +1315,16 @@ void SCineAssemblySchemaWindow::OnTreeItemTextCommitted(const FText& InText, ETe
 	const FString NewPath = FPaths::GetPath(OldPath) / InText.ToString();
 	TreeItem->Path = NewPath;
 
-	if (TreeItem->Type == FSchemaTreeItem::EItemType::Folder)
+	if (TSharedPtr<FSchemaContentHierarchyFolder> FolderItem = TreeItem->AsFolder())
 	{
 		// If this is a folder item, update the path of all of its children (recursively)
-		SetChildrenPathRecursive(TreeItem, NewPath);
+		SetChildrenPathRecursive(FolderItem, NewPath);
 
-		TreeItem->Parent->ChildFolders.Sort(UE::CineAssemblyTools::Private::SortTreeItems);
+		FolderItem->ParentFolder->ChildFolders.Sort(UE::CineAssemblyTools::Private::SortTreeItems);
 	}
 	else
 	{
-		TreeItem->Parent->ChildAssets.Sort(UE::CineAssemblyTools::Private::SortTreeItems);
+		TreeItem->ParentFolder->ChildAssets.Sort(UE::CineAssemblyTools::Private::SortTreeItems);
 	}
 
 	UpdateFolderList();
@@ -1297,15 +1333,16 @@ void SCineAssemblySchemaWindow::OnTreeItemTextCommitted(const FText& InText, ETe
 	TreeView->RequestTreeRefresh();
 }
 
-void SCineAssemblySchemaWindow::OnTreeItemTemplateChanged(const FAssetData& InAssetData, TSharedPtr<FSchemaTreeItem> TreeItem)
+void SCineAssemblySchemaWindow::OnTreeItemTemplateChanged(const FAssetData& InAssetData, TSharedPtr<ISchemaContentHierarchyTreeItem> TreeItem)
 {
 	// Early-out if this is not an asset item
-	if (TreeItem->Type != FSchemaTreeItem::EItemType::Asset)
+	TSharedPtr<FSchemaContentHierarchyAsset> AssetItem = TreeItem->AsAsset();
+	if (!AssetItem)
 	{
 		return;
 	}
 
-	const FSoftObjectPath OldSoftPath = TreeItem->AssetTemplate;
+	const FSoftObjectPath OldSoftPath = AssetItem->AssetTemplate;
 	const FSoftObjectPath NewSoftPath = InAssetData.GetSoftObjectPath();
 	// Early-out if the name has not actually changed
 	if (OldSoftPath == NewSoftPath)
@@ -1313,7 +1350,7 @@ void SCineAssemblySchemaWindow::OnTreeItemTemplateChanged(const FAssetData& InAs
 		return;
 	}
 
-	TreeItem->AssetTemplate = NewSoftPath;
+	AssetItem->AssetTemplate = NewSoftPath;
 	UpdateAssetList();
 	TreeView->RequestTreeRefresh();
 }
@@ -1338,7 +1375,7 @@ bool SCineAssemblySchemaWindow::ValidateAssetName(const FText& InText, FText& Ou
 	return FName::IsValidXName(InText.ToString(), InvalidCharacters, &OutErrorMessage);
 }
 
-bool SCineAssemblySchemaWindow::ValidateFolderName(const FText& InText, FText& OutErrorMessage, TSharedPtr<FSchemaTreeItem> TreeItem) const
+bool SCineAssemblySchemaWindow::ValidateFolderName(const FText& InText, FText& OutErrorMessage, TSharedPtr<FSchemaContentHierarchyFolder> FolderItem) const
 {
 	// Check for empty text string
 	if (InText.IsEmpty())
@@ -1357,10 +1394,10 @@ bool SCineAssemblySchemaWindow::ValidateFolderName(const FText& InText, FText& O
 	}
 
 	// Check for duplicate folder names
-	const FString ParentPath = FPaths::GetPath(TreeItem->Path);
-	TSharedPtr<FSchemaTreeItem> FoundItem = FindItemAtPathRecursive(RootItem, ParentPath / InText.ToString());
+	const FString ParentPath = FPaths::GetPath(FolderItem->Path);
+	TSharedPtr<FSchemaContentHierarchyFolder> MatchingFolder = FindFolderByPathRecursive(RootFolder, ParentPath / InText.ToString());
 
-	if (FoundItem && FoundItem != TreeItem)
+	if (MatchingFolder && MatchingFolder != FolderItem)
 	{
 		OutErrorMessage = LOCTEXT("DuplicateFolderErrorMessage", "A folder already exists at this location with this name");
 		return false;
@@ -1379,21 +1416,21 @@ void SCineAssemblySchemaWindow::DeleteAssetItem(TSharedPtr<FSchemaAssetParameter
 	AssetListView->RequestListRefresh();
 }
 
-void SCineAssemblySchemaWindow::DeleteTreeItem(TSharedPtr<FSchemaTreeItem> TreeItem)
+void SCineAssemblySchemaWindow::DeleteTreeItem(TSharedPtr<ISchemaContentHierarchyTreeItem> TreeItem)
 {
-	if (TreeItem->Type == FSchemaTreeItem::EItemType::Folder)
+	if (TSharedPtr<FSchemaContentHierarchyFolder> FolderItem = TreeItem->AsFolder())
 	{
-		if (ContainsRecursive(TreeItem, TopLevelAssemblyItem))
+		if (ContainsRecursive(FolderItem, TopLevelAssemblyItem))
 		{
-			TreeItem->Parent->ChildAssets.Add(TopLevelAssemblyItem);
-			TopLevelAssemblyItem->Parent = TreeItem->Parent;
+			FolderItem->ParentFolder->ChildAssets.Add(TopLevelAssemblyItem);
+			TopLevelAssemblyItem->ParentFolder = FolderItem->ParentFolder;
 		}
 
-		TreeItem->Parent->ChildFolders.Remove(TreeItem);
+		FolderItem->ParentFolder->ChildFolders.Remove(FolderItem);
 	}
 	else
 	{
-		TreeItem->Parent->ChildAssets.Remove(TreeItem);
+		TreeItem->ParentFolder->ChildAssets.Remove(TreeItem->AsAsset());
 	}
 
 	UpdateFolderList();
@@ -1402,17 +1439,17 @@ void SCineAssemblySchemaWindow::DeleteTreeItem(TSharedPtr<FSchemaTreeItem> TreeI
 	TreeView->RequestTreeRefresh();
 }
 
-bool SCineAssemblySchemaWindow::ContainsRecursive(TSharedPtr<FSchemaTreeItem> InParent, TSharedPtr<FSchemaTreeItem> InItem)
+bool SCineAssemblySchemaWindow::ContainsRecursive(TSharedPtr<FSchemaContentHierarchyFolder> InParentFolder, TSharedPtr<FSchemaContentHierarchyAsset> InAssetItem)
 {
-	if (InParent->ChildAssets.Contains(InItem))
+	if (InParentFolder->ChildAssets.Contains(InAssetItem))
 	{
 		return true;
 	}
 
 	bool bFoundRecursively = false;
-	for (TSharedPtr<FSchemaTreeItem> ChildFolder : InParent->ChildFolders)
+	for (TSharedPtr<FSchemaContentHierarchyFolder> ChildFolder : InParentFolder->ChildFolders)
 	{
-		bFoundRecursively |= ContainsRecursive(ChildFolder, InItem);
+		bFoundRecursively |= ContainsRecursive(ChildFolder, InAssetItem);
 	}
 
 	return bFoundRecursively;
@@ -1479,7 +1516,7 @@ void SCineAssemblySchemaWindow::UpdateFolderList()
 	const TArray<FString> CachedFolderList = SchemaToConfigure->FoldersToCreate;
 
 	TArray<FString> FolderList;
-	GetFolderListRecursive(RootItem, FolderList);
+	GetFolderListRecursive(RootFolder, FolderList);
 
 	if (FolderList != CachedFolderList)
 	{
@@ -1496,9 +1533,9 @@ void SCineAssemblySchemaWindow::UpdateAssetList()
 
 	TArray<FString> AssetList;
 	TMap<FString, FSoftObjectPath> Templates;
-	GetAssetListRecursive(RootItem, AssetList, Templates);
+	GetAssetListRecursive(RootFolder, AssetList, Templates);
 
-	SchemaToConfigure->DefaultAssemblyPath = TopLevelAssemblyItem->Parent->Path;
+	SchemaToConfigure->DefaultAssemblyPath = TopLevelAssemblyItem->ParentFolder->Path;
 
 	if (AssetList != CachedAssetList)
 	{
@@ -1674,42 +1711,43 @@ FReply SCineAssemblySchemaWindow::OnTreeRowDragDetected(const FGeometry& MyGeome
 {
 	if (MouseEvent.IsMouseButtonDown(EKeys::LeftMouseButton))
 	{
-		TArray<TSharedPtr<FSchemaTreeItem>> SelectedTreeItems = TreeView->GetSelectedItems();
-		if (SelectedTreeItems.Num() == 1)
+		TArray<TSharedPtr<ISchemaContentHierarchyTreeItem>> SelectedItems = TreeView->GetSelectedItems();
+		if (SelectedItems.Num() == 1)
 		{
-			TSharedPtr<FSchemaTreeItem> SelectedTreeItem = SelectedTreeItems[0];
+			TSharedPtr<ISchemaContentHierarchyTreeItem> SelectedTreeItem = SelectedItems[0];
 		
-			if (SelectedTreeItem == RootItem)
+			if (SelectedTreeItem == RootFolder)
 			{
 				return FReply::Handled();
 			}
 			
-			TSharedRef<FSchemaAssetDragDrop> Operation = FSchemaAssetDragDrop::New(FSchemaAssetParameters(FPaths::GetPathLeaf(SelectedTreeItem->Path), SelectedTreeItem->AssetTemplate));
+			FSoftObjectPath TemplatePath = SelectedTreeItem->AsAsset() ? SelectedTreeItem->AsAsset()->AssetTemplate : FSoftObjectPath();
+			TSharedRef<FSchemaAssetDragDrop> Operation = FSchemaAssetDragDrop::New(FSchemaAssetParameters(FPaths::GetPathLeaf(SelectedTreeItem->Path), TemplatePath));
 			Operation->SourceTreeItem = SelectedTreeItem;
 
-			if (SelectedTreeItem->Type == FSchemaTreeItem::EItemType::Asset)
+			if (TSharedPtr<FSchemaContentHierarchyAsset> AssetItem = SelectedTreeItem->AsAsset())
 			{
-				Operation->OnDropNotHandled.BindLambda([this, SelectedTreeItem](const FString& AssetName)
+				Operation->OnDropNotHandled.BindLambda([this, AssetItem](const FString& AssetName)
 					{
 						// If the tree item is not dropped in a valid place, restore it to its original place in the tree
-						SelectedTreeItem->Parent->ChildAssets.Add(SelectedTreeItem);
+						AssetItem->ParentFolder->ChildAssets.Add(AssetItem);
 						TreeView->RequestTreeRefresh();
 					});
 
 				// Remove this item from the tree while it is being dragged
-				SelectedTreeItem->Parent->ChildAssets.Remove(SelectedTreeItem);
+				SelectedTreeItem->ParentFolder->ChildAssets.Remove(AssetItem);
 			}
-			else
+			else if (TSharedPtr<FSchemaContentHierarchyFolder> FolderItem = SelectedTreeItem->AsFolder())
 			{
-				Operation->OnDropNotHandled.BindLambda([this, SelectedTreeItem](const FString& AssetName)
+				Operation->OnDropNotHandled.BindLambda([this, FolderItem](const FString& AssetName)
 					{
 						// If the tree item is not dropped in a valid place, restore it to its original place in the tree
-						SelectedTreeItem->Parent->ChildFolders.Add(SelectedTreeItem);
+						FolderItem->ParentFolder->ChildFolders.Add(FolderItem);
 						TreeView->RequestTreeRefresh();
 					});
 
 				// Remove this item from the tree while it is being dragged
-				SelectedTreeItem->Parent->ChildFolders.Remove(SelectedTreeItem);
+				FolderItem->ParentFolder->ChildFolders.Remove(FolderItem);
 			}
 
 			TreeView->RequestTreeRefresh();
@@ -1720,62 +1758,62 @@ FReply SCineAssemblySchemaWindow::OnTreeRowDragDetected(const FGeometry& MyGeome
 	return FReply::Unhandled();
 }
 
-FReply SCineAssemblySchemaWindow::OnTreeRowAcceptDrop(const FDragDropEvent& InDragDropEvent, EItemDropZone InDropZone, TSharedPtr<FSchemaTreeItem> InItem)
+FReply SCineAssemblySchemaWindow::OnTreeRowAcceptDrop(const FDragDropEvent& InDragDropEvent, EItemDropZone InDropZone, TSharedPtr<ISchemaContentHierarchyTreeItem> InItem)
 {
+	// Only Folder Items are valid drop targets
+	TSharedPtr<FSchemaContentHierarchyFolder> TargetFolder = InItem->AsFolder();
+	if (!TargetFolder)
+	{
+		return FReply::Unhandled();
+	}
+
 	TSharedPtr<FSchemaAssetDragDrop> Operation = InDragDropEvent.GetOperationAs<FSchemaAssetDragDrop>();
 	if (Operation.IsValid())
 	{
-		TSharedPtr<FSchemaTreeItem> NewItem;
+		TSharedPtr<ISchemaContentHierarchyTreeItem> DroppedItem;
 		if (Operation->SourceTreeItem)
 		{
-			NewItem = Operation->SourceTreeItem;
+			DroppedItem = Operation->SourceTreeItem;
 		}
 		else
 		{
 			// This must be a new item from the asset list, so we need to make a new tree item for it
-			NewItem = MakeShared<FSchemaTreeItem>();
-			NewItem->Type = FSchemaTreeItem::EItemType::Asset;
+			DroppedItem = MakeShared<FSchemaContentHierarchyAsset>();
 		}
 
 		// If an item with the same name already exists in the drop location, do not handle the drop (the dragged item will be reset to its original location)
-		if (NewItem->Type == FSchemaTreeItem::EItemType::Asset)
+		FString PendingDroppedItemPath = TargetFolder->Path / Operation->ItemName;
+		if (DroppedItem->AsAsset() && Algo::ContainsBy(TargetFolder->ChildAssets, PendingDroppedItemPath, &ISchemaContentHierarchyTreeItem::Path))
 		{
-			if (Algo::ContainsBy(InItem->ChildAssets, InItem->Path / Operation->ItemName, &FSchemaTreeItem::Path))
-			{
-				return FReply::Unhandled();
-			}
+			return FReply::Unhandled();
 		}
-		else if (NewItem->Type == FSchemaTreeItem::EItemType::Folder)
+		else if (DroppedItem->AsFolder() && Algo::ContainsBy(TargetFolder->ChildFolders, PendingDroppedItemPath, &ISchemaContentHierarchyTreeItem::Path))
 		{
-			if (Algo::ContainsBy(InItem->ChildFolders, InItem->Path / Operation->ItemName, &FSchemaTreeItem::Path))
-			{
-				return FReply::Unhandled();
-			}
+			return FReply::Unhandled();
 		}
 
-		const FString ParentPath = InItem->Path;
-		NewItem->Path = ParentPath / Operation->ItemName;
-		NewItem->AssetTemplate = Operation->Template;
+		DroppedItem->Path = PendingDroppedItemPath;
+		DroppedItem->ParentFolder = TargetFolder;
 
-		NewItem->Parent = InItem;
-
-		if (NewItem->Type == FSchemaTreeItem::EItemType::Asset)
+		if (TSharedPtr<FSchemaContentHierarchyAsset> DroppedAsset = DroppedItem->AsAsset())
 		{
-			InItem->ChildAssets.Add(NewItem);
-			InItem->ChildAssets.Sort(UE::CineAssemblyTools::Private::SortTreeItems);
+			DroppedAsset->AssetTemplate = Operation->Template;
+
+			TargetFolder->ChildAssets.Add(DroppedAsset);
+			TargetFolder->ChildAssets.Sort(UE::CineAssemblyTools::Private::SortTreeItems);
 		}
-		else
+		else if (TSharedPtr<FSchemaContentHierarchyFolder> DroppedFolder = DroppedItem->AsFolder())
 		{
-			InItem->ChildFolders.Add(NewItem);
+			TargetFolder->ChildFolders.Add(DroppedFolder);
 
-			SetChildrenPathRecursive(NewItem, NewItem->Path);
-			InItem->ChildFolders.Sort(UE::CineAssemblyTools::Private::SortTreeItems);
+			SetChildrenPathRecursive(DroppedFolder, DroppedFolder->Path);
+			TargetFolder->ChildFolders.Sort(UE::CineAssemblyTools::Private::SortTreeItems);
 		}
 
 		UpdateFolderList();
 		UpdateAssetList();
 
-		TreeView->SetItemExpansion(InItem, true);
+		ExpandTreeRecursive(TargetFolder);
 		TreeView->RequestTreeRefresh();
 	}
 
diff --git a/Engine/Plugins/VirtualProduction/CinematicAssemblyTools/Source/CineAssemblyToolsEditor/Private/UI/CineAssembly/SCineAssemblySchemaWindow.h b/Engine/Plugins/VirtualProduction/CinematicAssemblyTools/Source/CineAssemblyToolsEditor/Private/UI/CineAssembly/SCineAssemblySchemaWindow.h
--- a/Engine/Plugins/VirtualProduction/CinematicAssemblyTools/Source/CineAssemblyToolsEditor/Private/UI/CineAssembly/SCineAssemblySchemaWindow.h
+++ b/Engine/Plugins/VirtualProduction/CinematicAssemblyTools/Source/CineAssemblyToolsEditor/Private/UI/CineAssembly/SCineAssemblySchemaWindow.h
@@ -14,6 +14,10 @@
 
 struct FPropertyAndParent;
 
+struct ISchemaContentHierarchyTreeItem;
+struct FSchemaContentHierarchyFolder;
+struct FSchemaContentHierarchyAsset;
+
 /** Mode that configures the UI based on the intended user interactions with the Cine Assembly Schema asset */
 enum class ESchemaConfigMode : uint8
 {
@@ -39,38 +43,6 @@ public:
 	TSharedPtr<SNamingTokensEditableTextBox> TextBlock;
 };
 
-/** An entry in the content folder tree view with knowledge of its path, children, and parent */
-struct FSchemaTreeItem
-{
-	/** The types of items that can be represented in this tree view */
-	enum class EItemType : uint8
-	{
-		Asset,
-		Folder
-	};
-
-	/** The type of this tree item */
-	EItemType Type;
-
-	/** The relative path of this tree item */
-	FString Path;
-
-	/** The parent of this item in the tree */
-	TSharedPtr<FSchemaTreeItem> Parent;
-
-	/** The children of this item in the tree that are Asset types */
-	TArray<TSharedPtr<FSchemaTreeItem>> ChildAssets;
-
-	/** The children of this item in the tree that are Folder types */
-	TArray<TSharedPtr<FSchemaTreeItem>> ChildFolders;
-
-	/** Asset Template soft path when this item represents an asset. */
-	FSoftObjectPath AssetTemplate;
-
-	/** The text widget that displays this item's name and supports renaming */
-	TSharedPtr<SNamingTokensEditableTextBox> NameWidget;
-};
-
 /** Drag and Drop operation to handle dragging schema assets and folders  */
 class FSchemaAssetDragDrop : public FDragDropOperation
 {
@@ -91,7 +63,7 @@ public:
 	// End FDragDropOperation overrides
 
 	/** The source item being dragged. This could be null if the thing being dragged is from the Asset List (not the tree) */
-	TSharedPtr<FSchemaTreeItem> SourceTreeItem;
+	TSharedPtr<ISchemaContentHierarchyTreeItem> SourceTreeItem;
 
 	/** The name of the item being displayed, used by the decorator and to name the new tree item that will get made when dropped */
 	FString ItemName;
@@ -196,10 +168,10 @@ private:
 	FReply OnAssetListKeyDown(const FGeometry& MyGeometry, const FKeyEvent& InKeyEvent);
 
 	/** Generates the row widget for an entry in the tree view */
-	TSharedRef<ITableRow> OnGenerateTreeRow(TSharedPtr<FSchemaTreeItem> TreeItem, const TSharedRef<STableViewBase>& OwnerTable);
+	TSharedRef<ITableRow> OnGenerateTreeRow(TSharedPtr<ISchemaContentHierarchyTreeItem> TreeItem, const TSharedRef<STableViewBase>& OwnerTable);
 
 	/** Gets the children of the input tree view item to build additional tree rows */
-	void OnGetChildren(TSharedPtr<FSchemaTreeItem> TreeItem, TArray<TSharedPtr<FSchemaTreeItem>>& OutNodes);
+	void OnGetChildren(TSharedPtr<ISchemaContentHierarchyTreeItem> TreeItem, TArray<TSharedPtr<ISchemaContentHierarchyTreeItem>>& OutNodes);
 
 	/** Callback when the content tree view rebuilds itself, used to make the textbox in the most recently added row editable */
 	void OnTreeItemsRebuilt();
@@ -208,7 +180,7 @@ private:
 	TSharedPtr<SWidget> MakeContentTreeContextMenu();
 
 	/** Handles double clicks on the tree view to rename the item */
-	void OnTreeViewDoubleClick(TSharedPtr<FSchemaTreeItem> TreeItem);
+	void OnTreeViewDoubleClick(TSharedPtr<ISchemaContentHierarchyTreeItem> TreeItem);
 
 	/** Handles key presses on the tree view */
 	FReply OnTreeViewKeyDown(const FGeometry& MyGeometry, const FKeyEvent& InKeyEvent);
@@ -216,29 +188,29 @@ private:
 	/** Build the content tree from the list of folders and assets saved in the schema properties */
 	void InitializeContentTree();
 
-	/** Recursively expands every item in the tree view */
-	void ExpandTreeRecursive(TSharedPtr<FSchemaTreeItem> TreeItem) const;
+	/** Recursively expands every folder in the tree view */
+	void ExpandTreeRecursive(TSharedPtr<FSchemaContentHierarchyFolder> FolderItem) const;
 
 	/** Recursively get the path of every folder item in the tree view */
-	void GetFolderListRecursive(TSharedPtr<FSchemaTreeItem> TreeItem, TArray<FString>& FolderList) const;
+	void GetFolderListRecursive(TSharedPtr<FSchemaContentHierarchyFolder> FolderItem, TArray<FString>& OutFolderList) const;
 
 	/** Recursively get the path of every asset item in the tree view */
-	void GetAssetListRecursive(TSharedPtr<FSchemaTreeItem> TreeItem, TArray<FString>& AssetPathList, TMap<FString, FSoftObjectPath>& Templates) const;
+	void GetAssetListRecursive(TSharedPtr<FSchemaContentHierarchyFolder> FolderItem, TArray<FString>& OutAssetList, TMap<FString, FSoftObjectPath>& OutTemplates) const;
 
-	/** Returns the tree item whose path matches the input path */
-	TSharedPtr<FSchemaTreeItem> FindItemAtPathRecursive(TSharedPtr<FSchemaTreeItem> TreeItem, const FString& Path) const;
+	/** Returns the folder in the tree whose path matches the input path */
+	TSharedPtr<FSchemaContentHierarchyFolder> FindFolderByPathRecursive(TSharedPtr<FSchemaContentHierarchyFolder> FolderItem, const FString& FolderPath) const;
 
-	/** Modifies the path of the input tree item and all of its children (recursively) with the input new path */
-	void SetChildrenPathRecursive(TSharedPtr<FSchemaTreeItem> TreeItem, const FString& NewPath);
+	/** Modifies the path of the input folder item and all of its children (recursively) with the input new path */
+	void SetChildrenPathRecursive(TSharedPtr<FSchemaContentHierarchyFolder> FolderItem, const FString& NewPath);
 
 	/** Finds the row for the input item and puts its textblock into edit mode */
 	void EnterEditMode(TSharedPtr<FSchemaAssetParameters> ItemToRename);
 
 	/** Renames the tree item, and updates the paths of all of its children */
-	void OnTreeItemTextCommitted(const FText& InText, ETextCommit::Type InCommitType, TSharedPtr<FSchemaTreeItem> TreeItem);
+	void OnTreeItemTextCommitted(const FText& InText, ETextCommit::Type InCommitType, TSharedPtr<ISchemaContentHierarchyTreeItem> TreeItem);
 
 	/** Updates the Template object of the given tree item. */
-	void OnTreeItemTemplateChanged(const FAssetData& InAssetData, TSharedPtr<FSchemaTreeItem> TreeItem);
+	void OnTreeItemTemplateChanged(const FAssetData& InAssetData, TSharedPtr<ISchemaContentHierarchyTreeItem> TreeItem);
 
 	/** Callback when one of the properties of the schema being configured changes */
 	void OnSchemaPropertiesChanged(const FPropertyChangedEvent& PropertyChangedEvent);
@@ -247,13 +219,13 @@ private:
 	bool ValidateAssetName(const FText& InText, FText& OutErrorMessage) const;
 
 	/** Validate the text entered by the user to ensure it will be a valid folder name */
-	bool ValidateFolderName(const FText& InText, FText& OutErrorMessage, TSharedPtr<FSchemaTreeItem> TreeItem) const;
+	bool ValidateFolderName(const FText& InText, FText& OutErrorMessage, TSharedPtr<FSchemaContentHierarchyFolder> FolderItem) const;
 
 	/** Remove the input item from the asset list */
 	void DeleteAssetItem(TSharedPtr<FSchemaAssetParameters> InItem);
 
 	/** Remove the input item and its children from the tree view */
-	void DeleteTreeItem(TSharedPtr<FSchemaTreeItem> TreeItem);
+	void DeleteTreeItem(TSharedPtr<ISchemaContentHierarchyTreeItem> TreeItem);
 
 	/** Begins a drag and drop event to drag an item out of the asset list view */
 	FReply OnAssetRowDragDetected(const FGeometry& MyGeometry, const FPointerEvent& MouseEvent);
@@ -262,13 +234,13 @@ private:
 	FReply OnTreeRowDragDetected(const FGeometry& MyGeometry, const FPointerEvent& MouseEvent);
 
 	/** Callback when an item is dropped onto a folder in the content tree view to add it to the children of that row */
-	FReply OnTreeRowAcceptDrop(const FDragDropEvent& InDragDropEvent, EItemDropZone InDropZone, TSharedPtr<FSchemaTreeItem> InItem);
+	FReply OnTreeRowAcceptDrop(const FDragDropEvent& InDragDropEvent, EItemDropZone InDropZone, TSharedPtr<ISchemaContentHierarchyTreeItem> InItem);
 
-	/** Generates a unique folder path name, assuming the input item will be the parent */
-	FString MakeUniqueFolderPath(TSharedPtr<FSchemaTreeItem> InItem);
+	/** Generates a unique folder path name, assuming the input folder will be the parent folder */
+	FString MakeUniqueFolderPath(TSharedPtr<FSchemaContentHierarchyFolder> FolderItem);
 
-	/** Returns true if the input parent, or any of its child folders, contains in the input item */
-	bool ContainsRecursive(TSharedPtr<FSchemaTreeItem> InParent, TSharedPtr<FSchemaTreeItem> InItem);
+	/** Returns true if the input parent folder, or any of its child folders, contains in the input asset item */
+	bool ContainsRecursive(TSharedPtr<FSchemaContentHierarchyFolder> InParentFolder, TSharedPtr<FSchemaContentHierarchyAsset> InAssetItem);
 
 	/** Update the list of folders to create from the current set in the tree view */
 	void UpdateFolderList();
@@ -318,19 +290,19 @@ private:
 	TSharedPtr<SListView<TSharedPtr<FSchemaAssetParameters>>> AssetListView;
 
 	/** Items source for the tree view */
-	TArray<TSharedPtr<FSchemaTreeItem>> TreeItems;
+	TArray<TSharedPtr<ISchemaContentHierarchyTreeItem>> TreeItems;
 
-	/** The root item in the tree view */
-	TSharedPtr<FSchemaTreeItem> RootItem;
+	/** The root folder in the tree view */
+	TSharedPtr<FSchemaContentHierarchyFolder> RootFolder;
 
 	/** An item representing where the top-level assembly should be created */
-	TSharedPtr<FSchemaTreeItem> TopLevelAssemblyItem;
+	TSharedPtr<FSchemaContentHierarchyAsset> TopLevelAssemblyItem;
 
 	/** The tree view of content folders and assets for this schema */
-	TSharedPtr<STreeView<TSharedPtr<FSchemaTreeItem>>> TreeView;
+	TSharedPtr<STreeView<TSharedPtr<ISchemaContentHierarchyTreeItem>>> TreeView;
 
 	/** The mostly recently added tree item, used to allow the user to immediately rename the item after the tree view updates to show it */
-	TSharedPtr<FSchemaTreeItem> MostRecentlyAddedItem;
+	TSharedPtr<ISchemaContentHierarchyTreeItem> MostRecentlyAddedItem;
 
 	/** The sequencer used to display / edit the template sequence */
 	TSharedPtr<ISequencer> Sequencer;
