ShotManagement: Add support to the Schema Timeline Template to allow SubAssembly Tracks to also customize the behavior of Shot Tracks in the Schema Sequencer tool.
#rnx

#Synced-CL 49044222

diff --git a/Engine/Plugins/VirtualProduction/CinematicAssemblyTools/Source/CineAssemblyTools/Private/CineAssembly.cpp b/Engine/Plugins/VirtualProduction/CinematicAssemblyTools/Source/CineAssemblyTools/Private/CineAssembly.cpp
--- a/Engine/Plugins/VirtualProduction/CinematicAssemblyTools/Source/CineAssemblyTools/Private/CineAssembly.cpp
+++ b/Engine/Plugins/VirtualProduction/CinematicAssemblyTools/Source/CineAssemblyTools/Private/CineAssembly.cpp
@@ -12,10 +12,12 @@
 #include "MovieSceneSubAssemblySection.h"
 #include "MovieSceneSubAssemblyTrack.h"
 #include "NamingTokensEngineSubsystem.h"
+#include "Sections/MovieSceneCinematicShotSection.h"
 #include "Sections/MovieSceneSubSection.h"
 #include "Serialization/JsonReader.h"
 #include "Serialization/JsonSerializer.h"
 #include "Serialization/JsonWriter.h"
+#include "Tracks/MovieSceneCinematicShotTrack.h"
 #include "Tracks/MovieSceneEventTrack.h"
 #include "Tracks/MovieSceneSubTrack.h"
 #include "UniversalObjectLocator.h"
@@ -134,8 +136,8 @@ ETrackSupport UCineAssembly::IsTrackSupportedImpl(TSubclassOf<class UMovieSceneT
 	if (GetTypedOuter<UCineAssemblySchema>())
 	{
 		// Event tracks are not supported in the Schema Sequencer
-		// Normal Subsequence Tracks are not supported in the Schema Sequencer, because the custom SubAssembly Track is used instead
-		if ((InTrackClass == UMovieSceneEventTrack::StaticClass()) || (InTrackClass == UMovieSceneSubTrack::StaticClass()))
+		// Subsequence Tracks and Cinematic Shot Tracks are not supported in the Schema Sequencer, because the custom SubAssembly Track is used instead
+		if ((InTrackClass == UMovieSceneEventTrack::StaticClass()) || (InTrackClass == UMovieSceneSubTrack::StaticClass()) || (InTrackClass == UMovieSceneCinematicShotTrack::StaticClass()))
 		{
 			return ETrackSupport::NotSupported;
 		}
@@ -215,20 +217,30 @@ void UCineAssembly::InitializeFromTemplate(ULevelSequence* Template)
 		}
 
 #if WITH_EDITOR
-		// Replace any SubAssemblyTracks we find with normal SubTracks (and create any SubAssemblies needed by those tracks)
+		// Replace any SubAssemblyTracks we find with either a MovieSceneSubTrack or MovieSceneCinematicShotTrack, depending on the TrackType
 		while (UMovieSceneSubAssemblyTrack* SubAssemblyTrack = MovieScene->FindTrack<UMovieSceneSubAssemblyTrack>())
 		{
-			// Add a new SubTrack to the MovieScene. The SubAssemblyTrack will be removed from the MovieScene at the bottom of this loop.
-			UMovieSceneSubTrack* SubTrack = Cast<UMovieSceneSubTrack>(MovieScene->AddTrack(UMovieSceneSubTrack::StaticClass()));
-			SubTrack->SetDisplayName(SubAssemblyTrack->GetDisplayName());
+			// Add a new track to the MovieScene. The SubAssemblyTrack will be removed from the MovieScene at the bottom of this loop.
+			UClass* TrackClass = SubAssemblyTrack->IsSubsequenceTrack() ? UMovieSceneSubTrack::StaticClass() : UMovieSceneCinematicShotTrack::StaticClass();
+			UMovieSceneNameableTrack* NewTrack = Cast<UMovieSceneNameableTrack>(MovieScene->AddTrack(TrackClass));
+
+			NewTrack->SetDisplayName(SubAssemblyTrack->GetDisplayName());
 
 			for (UMovieSceneSection* Section : SubAssemblyTrack->GetAllSections())
 			{
 				if (UMovieSceneSubAssemblySection* SubAssemblySection = Cast<UMovieSceneSubAssemblySection>(Section))
 				{
-					// Create a new SubSection, with the relevant properties copied from the SubAssemblySection, and add it to the new SubTrack
-					UMovieSceneSubSection* NewSubSection = NewObject<UMovieSceneSubSection>(SubTrack, NAME_None, SubAssemblySection->GetFlags(), SubAssemblySection);
-					SubTrack->AddSection(*NewSubSection);
+					// Create a new section, with the relevant properties copied from the SubAssemblySection, and add it to the new track
+					UClass* SectionClass = SubAssemblyTrack->IsSubsequenceTrack() ? UMovieSceneSubSection::StaticClass() : UMovieSceneCinematicShotSection::StaticClass();
+					UMovieSceneSubSection* NewSubSection = NewObject<UMovieSceneSubSection>(NewTrack, SectionClass, NAME_None, SubAssemblySection->GetFlags(), SubAssemblySection);
+
+					NewTrack->AddSection(*NewSubSection);
+
+					// Remove the tint from the section that was used in the Schema editor to differentiate between template and reference sections
+					if (NewSubSection->GetColorTint() == UMovieSceneSubAssemblySection::TemplateSectionColorTint)
+					{
+						NewSubSection->SetColorTint(UMovieSceneSubAssemblySection::ReferenceSectionColorTint);
+					}
 
 					if (SubAssemblySection->IsTemplateSection())
 					{
diff --git a/Engine/Plugins/VirtualProduction/CinematicAssemblyTools/Source/CineAssemblyTools/Private/MovieSceneSubAssemblySection.cpp b/Engine/Plugins/VirtualProduction/CinematicAssemblyTools/Source/CineAssemblyTools/Private/MovieSceneSubAssemblySection.cpp
--- a/Engine/Plugins/VirtualProduction/CinematicAssemblyTools/Source/CineAssemblyTools/Private/MovieSceneSubAssemblySection.cpp
+++ b/Engine/Plugins/VirtualProduction/CinematicAssemblyTools/Source/CineAssemblyTools/Private/MovieSceneSubAssemblySection.cpp
@@ -5,6 +5,7 @@
 #include "CineAssembly.h"
 #include "CineAssemblySchema.h"
 #include "LevelSequence.h"
+#include "MovieSceneSubAssemblyTrack.h"
 
 #define LOCTEXT_NAMESPACE "MovieSceneSubAssemblySection"
 
@@ -17,17 +18,29 @@ UMovieSceneSubAssemblySection::UMovieSceneSubAssemblySection(const FObjectInitia
 {
 }
 
-bool UMovieSceneSubAssemblySection::IsReferenceSection()
+bool UMovieSceneSubAssemblySection::IsReferenceSection() const
 { 
 	return SectionType == ESubAssemblySectionType::Reference; 
 }
 
-bool UMovieSceneSubAssemblySection::IsTemplateSection()
+bool UMovieSceneSubAssemblySection::IsTemplateSection() const
 { 
 	return SectionType == ESubAssemblySectionType::Template;
 }
 
-UObject* UMovieSceneSubAssemblySection::GetAssemblyTemplate()
+bool UMovieSceneSubAssemblySection::IsSubsequenceSection() const
+{
+	UMovieSceneSubAssemblyTrack* SubAssemblyTrack = GetTypedOuter<UMovieSceneSubAssemblyTrack>();
+	return (SubAssemblyTrack && SubAssemblyTrack->IsSubsequenceTrack());
+}
+
+bool UMovieSceneSubAssemblySection::IsCinematicShotSection() const
+{
+	UMovieSceneSubAssemblyTrack* SubAssemblyTrack = GetTypedOuter<UMovieSceneSubAssemblyTrack>();
+	return (SubAssemblyTrack && SubAssemblyTrack->IsCinematicShotTrack());
+}
+
+UObject* UMovieSceneSubAssemblySection::GetAssemblyTemplate() const
 {
 	return AssemblyTemplate;
 }
diff --git a/Engine/Plugins/VirtualProduction/CinematicAssemblyTools/Source/CineAssemblyTools/Private/MovieSceneSubAssemblyTrack.cpp b/Engine/Plugins/VirtualProduction/CinematicAssemblyTools/Source/CineAssemblyTools/Private/MovieSceneSubAssemblyTrack.cpp
--- a/Engine/Plugins/VirtualProduction/CinematicAssemblyTools/Source/CineAssemblyTools/Private/MovieSceneSubAssemblyTrack.cpp
+++ b/Engine/Plugins/VirtualProduction/CinematicAssemblyTools/Source/CineAssemblyTools/Private/MovieSceneSubAssemblyTrack.cpp
@@ -4,6 +4,8 @@
 
 #include "MovieSceneSubAssemblySection.h"
 
+#define LOCTEXT_NAMESPACE "MovieSceneSubAssemblyTrack"
+
 UMovieSceneSubAssemblyTrack::UMovieSceneSubAssemblyTrack(const FObjectInitializer& ObjectInitializer)
 	: Super(ObjectInitializer)
 {
@@ -13,3 +15,32 @@ UMovieSceneSection* UMovieSceneSubAssemblyTrack::CreateNewSection()
 {
 	return NewObject<UMovieSceneSubAssemblySection>(this, NAME_None, RF_Transactional);
 }
+
+bool UMovieSceneSubAssemblyTrack::IsSubsequenceTrack() const
+{
+	return TrackType == ESubAssemblyTrackType::SubsequenceTrack;
+}
+
+bool UMovieSceneSubAssemblyTrack::IsCinematicShotTrack() const
+{
+	return TrackType == ESubAssemblyTrackType::CinematicShotTrack;
+}
+
+void UMovieSceneSubAssemblyTrack::SetTrackType(ESubAssemblyTrackType InTrackType)
+{
+	TrackType = InTrackType;
+}
+
+#if WITH_EDITOR
+FText UMovieSceneSubAssemblyTrack::GetDefaultDisplayName() const
+{
+	return IsSubsequenceTrack() ? LOCTEXT("SubsequenceTrackName", "Subsequences") : LOCTEXT("ShotTrackName", "Shots");
+}
+
+FName UMovieSceneSubAssemblyTrack::GetTrackIconBrush() const
+{
+	return IsSubsequenceTrack() ? TEXT("Sequencer.Tracks.Sub") : TEXT("Sequencer.Tracks.CinematicShot");
+}
+#endif // WITH_EDITOR
+
+#undef LOCTEXT_NAMESPACE
diff --git a/Engine/Plugins/VirtualProduction/CinematicAssemblyTools/Source/CineAssemblyTools/Public/MovieSceneSubAssemblySection.h b/Engine/Plugins/VirtualProduction/CinematicAssemblyTools/Source/CineAssemblyTools/Public/MovieSceneSubAssemblySection.h
--- a/Engine/Plugins/VirtualProduction/CinematicAssemblyTools/Source/CineAssemblyTools/Public/MovieSceneSubAssemblySection.h
+++ b/Engine/Plugins/VirtualProduction/CinematicAssemblyTools/Source/CineAssemblyTools/Public/MovieSceneSubAssemblySection.h
@@ -2,7 +2,7 @@
 
 #pragma once
 
-#include "Sections/MovieSceneSubSection.h"
+#include "Sections/MovieSceneCinematicShotSection.h"
 
 #include "MovieSceneSubAssemblySection.generated.h"
 
@@ -17,10 +17,11 @@ enum class ESubAssemblySectionType : uint8
 };
 
 /**
- * Implements a section in SubAssembly Tracks
+ * The UMovieSceneSubAssemblySection is only supported in the Template Assembly of a CineAssemblySchema as a placeholder section.
+ * The Track containing this section has a TrackType that indicates which section should be created when creating an Assembly from a Schema Template.
  */
 UCLASS(BlueprintType, MinimalAPI)
-class UMovieSceneSubAssemblySection : public UMovieSceneSubSection
+class UMovieSceneSubAssemblySection : public UMovieSceneCinematicShotSection
 {
 	GENERATED_BODY()
 
@@ -34,13 +35,19 @@ public:
 #endif //WITH_EDITOR
 
 	/** Returns true if this section is used as a reference to an existing sequence */
-	UE_API bool IsReferenceSection();
+	UE_API bool IsReferenceSection() const;
 
 	/** Returns true if this section is used as a template to create a new sequence */
-	UE_API bool IsTemplateSection();
+	UE_API bool IsTemplateSection() const;
+
+	/** Returns true if this section is a placeholder for a UMovieSceneSubSection */
+	UE_API bool IsSubsequenceSection() const;
+
+	/** Returns true if this section is a placeholder for a UMovieSceneCinematicShotSection */
+	UE_API bool IsCinematicShotSection() const;
 
 	/** Returns the Assembly Template object */
-	UE_API UObject* GetAssemblyTemplate();
+	UE_API UObject* GetAssemblyTemplate() const;
 
 #if WITH_EDITOR
 	/** Sets the Assembly Template and updates the SequenceName to match the new template */
@@ -59,14 +66,13 @@ public:
 	/** Name of the protected AssemblyTemplate property */
 	static UE_API const FName AssemblyTemplatePropertyName;
 
+	static const FColor ReferenceSectionColorTint;
+	static const FColor TemplateSectionColorTint;
+
 protected:
 	/** The object that will be used as a template to create a new sequence in this section. Can be a Level Sequence, Cine Assembly, or Schema */
 	UPROPERTY(EditAnywhere, Category = "Sequence")
 	TObjectPtr<UObject> AssemblyTemplate;
-
-private:
-	static const FColor ReferenceSectionColorTint;
-	static const FColor TemplateSectionColorTint;
 };
 
 #undef UE_API
diff --git a/Engine/Plugins/VirtualProduction/CinematicAssemblyTools/Source/CineAssemblyTools/Public/MovieSceneSubAssemblyTrack.h b/Engine/Plugins/VirtualProduction/CinematicAssemblyTools/Source/CineAssemblyTools/Public/MovieSceneSubAssemblyTrack.h
--- a/Engine/Plugins/VirtualProduction/CinematicAssemblyTools/Source/CineAssemblyTools/Public/MovieSceneSubAssemblyTrack.h
+++ b/Engine/Plugins/VirtualProduction/CinematicAssemblyTools/Source/CineAssemblyTools/Public/MovieSceneSubAssemblyTrack.h
@@ -8,9 +8,19 @@
 
 #define UE_API CINEASSEMBLYTOOLS_API
 
+/** 
+ * Indicates the underlying track type for which this track is a placeholder
+ */
+UENUM()
+enum class ESubAssemblyTrackType : uint8
+{
+	SubsequenceTrack,  /** This track is a placeholder for a UMovieSceneSubTrack */
+	CinematicShotTrack /** This track is a placeholder for a UMovieSceneCinematicShotTrack */
+};
+
 /**
- * A Track to hold SubAssemblies within a parent CineAssembly.
- * Used by the Schema Sequencer tool to support custom behavior above and beyond normal SubTracks
+ * The UMovieSceneSubAssemblyTrack is only supported in the Template Assembly of a CineAssemblySchema as a placeholder track.
+ * The Track Type indicates which underlying track should be created in place of this track when creating an Assembly from a Schema Template.
  */
 UCLASS(BlueprintType, MinimalAPI)
 class UMovieSceneSubAssemblyTrack : public UMovieSceneSubTrack
@@ -23,6 +33,31 @@ public:
 	//~ Begin UMovieSceneTrack interface
 	UE_API virtual UMovieSceneSection* CreateNewSection() override;
 	//~ End UMovieSceneTrack interface
+
+	//~ Begin UMovieSceneNameableTrack interface
+#if WITH_EDITOR
+	UE_API virtual FText GetDefaultDisplayName() const override;
+#endif // WITH_EDITOR
+	//~ End UMovieSceneNameableTrack interface
+
+	/** Returns true is this track represents a UMovieSceneSubTrack */
+	UE_API bool IsSubsequenceTrack() const;
+
+	/** Returns true is this track represents a UMovieSceneCinematicShotTrack */
+	UE_API bool IsCinematicShotTrack() const;
+
+	/** Sets the underlying track type that this track is meant to represent */
+	UE_API void SetTrackType(ESubAssemblyTrackType InTrackType);
+
+#if WITH_EDITOR
+	/** Get the icon brush for this track, determined by the track type */
+	UE_API FName GetTrackIconBrush() const;
+#endif // WITH_EDITOR
+
+private:
+	/** Indicates the underlying track type for which this track is a placeholder */
+	UPROPERTY()
+	ESubAssemblyTrackType TrackType = ESubAssemblyTrackType::CinematicShotTrack;
 };
 
 #undef UE_API
diff --git a/Engine/Plugins/VirtualProduction/CinematicAssemblyTools/Source/CineAssemblyToolsEditor/Private/Sequencer/SubAssemblySection.cpp b/Engine/Plugins/VirtualProduction/CinematicAssemblyTools/Source/CineAssemblyToolsEditor/Private/Sequencer/SubAssemblySection.cpp
--- a/Engine/Plugins/VirtualProduction/CinematicAssemblyTools/Source/CineAssemblyToolsEditor/Private/Sequencer/SubAssemblySection.cpp
+++ b/Engine/Plugins/VirtualProduction/CinematicAssemblyTools/Source/CineAssemblyToolsEditor/Private/Sequencer/SubAssemblySection.cpp
@@ -25,13 +25,24 @@ FText FSubAssemblySection::GetSectionTitle() const
 	return FText::GetEmpty();
 }
 
+float FSubAssemblySection::GetSectionHeight(const UE::Sequencer::FViewDensityInfo& ViewDensity) const
+{
+	static constexpr float ShotSectionHeight = 98.0f;
+	static constexpr float SubsequenceSectionHeight = 50.0f;
+
+	return MovieSceneSubAssemblySection->IsCinematicShotSection() ? ShotSectionHeight : SubsequenceSectionHeight;
+}
+
 TSharedRef<SWidget> FSubAssemblySection::GenerateSectionWidget()
 {
+	const FMargin SectionTextPadding = (MovieSceneSubAssemblySection->IsCinematicShotSection()) ? FMargin(0.0f, 12.0f, 0.0f, 0.0f) : FMargin(0.0f, 2.0f, 0.0f, 0.0f);
+	const FMargin OriginTextPadding = (MovieSceneSubAssemblySection->IsCinematicShotSection()) ? FMargin(0.0f, 15.0f, 0.0f, 0.0f) : FMargin(0.0f, 5.0f, 0.0f, 0.0f);
+
 	return SNew(SHorizontalBox)
 		+ SHorizontalBox::Slot()
 		.AutoWidth()
 		.VAlign(VAlign_Top)
-		.Padding(0.0f, 2.0f, 0.0f, 0.0f)
+		.Padding(SectionTextPadding)
 		[
 			// TODO: Validate that the tokenized text could possibly be a valid asset name
 			SAssignNew(EditableTextBlock, SNamingTokensEditableTextBox)
@@ -45,7 +56,7 @@ TSharedRef<SWidget> FSubAssemblySection::GenerateSectionWidget()
 
 		+ SHorizontalBox::Slot()
 		.FillWidth(1.0f)
-		.Padding(0.0f, 5.0f, 0.0f, 0.0f)
+		.Padding(OriginTextPadding)
 		[
 			SNew(STextBlock)
 				.Text(this, &FSubAssemblySection::GetOriginText)
@@ -53,6 +64,37 @@ TSharedRef<SWidget> FSubAssemblySection::GenerateSectionWidget()
 		];
 }
 
+int32 FSubAssemblySection::OnPaintSection(FSequencerSectionPainter& InPainter) const
+{
+	InPainter.LayerId = InPainter.PaintSectionBackground();
+
+	// Draw the film border used by shot tracks
+	if (MovieSceneSubAssemblySection->IsCinematicShotSection())
+	{
+		static const FSlateBrush* FilmBorder = FAppStyle::GetBrush("Sequencer.Section.FilmBorder");
+
+		FVector2D LocalHeaderSize = InPainter.HeaderGeometry.GetLocalSize();
+
+		FSlateDrawElement::MakeBox(
+			InPainter.DrawElements,
+			InPainter.LayerId++,
+			InPainter.SectionGeometry.ToPaintGeometry(FVector2D(LocalHeaderSize.X - 2.f, 7.f), FSlateLayoutTransform(FVector2D(1.f, 4.f))),
+			FilmBorder,
+			InPainter.bParentEnabled ? ESlateDrawEffect::None : ESlateDrawEffect::DisabledEffect
+		);
+
+		FSlateDrawElement::MakeBox(
+			InPainter.DrawElements,
+			InPainter.LayerId++,
+			InPainter.SectionGeometry.ToPaintGeometry(FVector2D(LocalHeaderSize.X - 2.f, 7.f), FSlateLayoutTransform(FVector2D(1.f, LocalHeaderSize.Y - 11.f))),
+			FilmBorder,
+			InPainter.bParentEnabled ? ESlateDrawEffect::None : ESlateDrawEffect::DisabledEffect
+		);
+	}
+
+	return InPainter.LayerId;
+}
+
 FReply FSubAssemblySection::OnSectionDoubleClicked(const FGeometry& SectionGeometry, const FPointerEvent& MouseEvent)
 {
 	// TODO: Check the geometry to determine what double clicking should do.
diff --git a/Engine/Plugins/VirtualProduction/CinematicAssemblyTools/Source/CineAssemblyToolsEditor/Private/Sequencer/SubAssemblySection.h b/Engine/Plugins/VirtualProduction/CinematicAssemblyTools/Source/CineAssemblyToolsEditor/Private/Sequencer/SubAssemblySection.h
--- a/Engine/Plugins/VirtualProduction/CinematicAssemblyTools/Source/CineAssemblyToolsEditor/Private/Sequencer/SubAssemblySection.h
+++ b/Engine/Plugins/VirtualProduction/CinematicAssemblyTools/Source/CineAssemblyToolsEditor/Private/Sequencer/SubAssemblySection.h
@@ -20,7 +20,9 @@ public:
 	//~ Begin ISequencerSection interface
 	virtual UMovieSceneSection* GetSectionObject() override;
 	virtual FText GetSectionTitle() const override;
+	virtual float GetSectionHeight(const UE::Sequencer::FViewDensityInfo& ViewDensity) const override;
 	virtual TSharedRef<SWidget> GenerateSectionWidget() override;
+	virtual int32 OnPaintSection(FSequencerSectionPainter& InPainter) const override;
 	virtual FReply OnSectionDoubleClicked(const FGeometry& SectionGeometry, const FPointerEvent& MouseEvent) override;
 	virtual void BuildSectionContextMenu(FMenuBuilder& MenuBuilder, const FGuid& ObjectBinding) override;
 	virtual void CustomizePropertiesDetailsView(TSharedRef<IDetailsView> DetailsView, const FSequencerSectionPropertyDetailsViewCustomizationParams& InParams) const override;
diff --git a/Engine/Plugins/VirtualProduction/CinematicAssemblyTools/Source/CineAssemblyToolsEditor/Private/Sequencer/SubAssemblyTrackEditor.cpp b/Engine/Plugins/VirtualProduction/CinematicAssemblyTools/Source/CineAssemblyToolsEditor/Private/Sequencer/SubAssemblyTrackEditor.cpp
--- a/Engine/Plugins/VirtualProduction/CinematicAssemblyTools/Source/CineAssemblyToolsEditor/Private/Sequencer/SubAssemblyTrackEditor.cpp
+++ b/Engine/Plugins/VirtualProduction/CinematicAssemblyTools/Source/CineAssemblyToolsEditor/Private/Sequencer/SubAssemblyTrackEditor.cpp
@@ -11,6 +11,7 @@
 #include "MovieSceneToolsProjectSettings.h"
 #include "MVVM/Extensions/ITrackExtension.h"
 #include "MVVM/Selection/Selection.h"
+#include "MVVM/ViewModels/OutlinerColumns/OutlinerColumnTypes.h"
 #include "MVVM/Views/ViewUtilities.h"
 #include "Sections/MovieSceneSubSection.h"
 #include "SequencerSettings.h"
@@ -70,12 +71,12 @@ bool FSubAssemblyTrackEditor::HandleAssetAdded(UObject* Asset, const FGuid& Targ
 		return false;
 	}
 
-	if (!SupportsSequence(Sequence))
+	if (!Sequence->GetMovieScene())
 	{
 		return false;
 	}
 
-	if (!CanHandleAssetAdded(Sequence))
+	if (!SupportsSequence(Sequence))
 	{
 		return false;
 	}
@@ -115,11 +116,101 @@ TSharedRef<ISequencerSection> FSubAssemblyTrackEditor::MakeSectionInterface(UMov
 	return MakeShared<FSubAssemblySection>(GetSequencer(), SubAssemblySection);
 }
 
+void FSubAssemblyTrackEditor::BuildAddTrackMenu(FMenuBuilder& MenuBuilder)
+{
+	// This Track Editor mimics both the Subsequence Track Editor and Cinematic Shot Track Editor for the purpose of creating the Schema's template sequence
+	MenuBuilder.AddMenuEntry(
+		LOCTEXT("ShotTrackName", "Shot Track"),
+		LOCTEXT("ShotTrackToolTip", "A cinematic shot track."),
+		FSlateIcon(FAppStyle::GetAppStyleSetName(), "Sequencer.Tracks.CinematicShot"),
+		FUIAction(FExecuteAction::CreateLambda([this]() 
+			{
+				AddSubAssemblyTrack(ESubAssemblyTrackType::CinematicShotTrack);
+			}))
+	);
+
+	MenuBuilder.AddMenuEntry(
+		LOCTEXT("SubTrackName", "Subsequence Track"),
+		LOCTEXT("SubTrackToolTip", "A track that can contain other sequences."),
+		FSlateIcon(FAppStyle::GetAppStyleSetName(), "Sequencer.Tracks.Sub"),
+		FUIAction(FExecuteAction::CreateLambda([this]()
+			{
+				AddSubAssemblyTrack(ESubAssemblyTrackType::SubsequenceTrack);
+			}))
+	);
+}
+
+UMovieSceneSubAssemblyTrack* FSubAssemblyTrackEditor::AddSubAssemblyTrack(ESubAssemblyTrackType TrackType)
+{
+	UMovieScene* FocusedMovieScene = GetFocusedMovieScene();
+	if (!FocusedMovieScene)
+	{
+		return nullptr;
+	}
+
+	// Regardless of the track type, a new MovieSceneSubAssemblyTrack is added to the Schema's template sequence
+	UMovieSceneSubAssemblyTrack* NewTrack = Cast<UMovieSceneSubAssemblyTrack>(FocusedMovieScene->AddTrack(UMovieSceneSubAssemblyTrack::StaticClass()));
+	NewTrack->SetTrackType(TrackType);
+
+	if (GetSequencer().IsValid())
+	{
+		GetSequencer()->OnAddTrack(NewTrack, FGuid());
+	}
+
+	return NewTrack;
+}
+
 TSharedPtr<SWidget> FSubAssemblyTrackEditor::BuildOutlinerEditWidget(const FGuid& ObjectBinding, UMovieSceneTrack* Track, const FBuildEditWidgetParams& Params)
 {
 	return UE::Sequencer::MakeAddButton(GetSubTrackName(), FOnGetContent::CreateRaw(this, &FSubAssemblyTrackEditor::BuildAddSubSequenceComboButtonContent, Params.ViewModel.AsWeak()), Params.ViewModel);
 }
 
+TSharedPtr<SWidget> FSubAssemblyTrackEditor::BuildOutlinerColumnWidget(const FBuildColumnWidgetParams& Params, const FName& ColumnName)
+{
+	// This label column widget resembles an SOutlinerItemViewBase, but allows the different track types to each support their own icon and text label
+	if (ColumnName == UE::Sequencer::FCommonOutlinerNames::Label)
+	{
+		UE::Sequencer::TViewModelPtr<UE::Sequencer::ITrackExtension> TrackModel = Params.ViewModel->FindAncestorOfType<UE::Sequencer::ITrackExtension>(true);
+		if (!TrackModel)
+		{
+			return SNullWidget::NullWidget;
+		}
+
+		UMovieSceneSubAssemblyTrack* Track = Cast<UMovieSceneSubAssemblyTrack>(TrackModel->GetTrack());
+		if (!Track)
+		{
+			return SNullWidget::NullWidget;
+		}
+
+		return SNew(SHorizontalBox)
+
+			+ SHorizontalBox::Slot()
+			.AutoWidth()
+			[
+				SNew(SExpanderArrow, Params.TreeViewRow)
+					.IndentAmount(10.0f)
+			]
+
+			+ SHorizontalBox::Slot()
+			.AutoWidth()
+			.Padding(4.0f, 0.0f)
+			[
+				SNew(SImage)
+					.Image(FAppStyle::Get().GetBrush(Track->GetTrackIconBrush()))
+			]
+
+			+ SHorizontalBox::Slot()
+			.AutoWidth()
+			[
+				SNew(STextBlock)
+					.Text(Track->GetDefaultDisplayName())
+					.Font(FAppStyle::GetFontStyle("Sequencer.AnimationOutliner.RegularFont"))
+			];
+	}
+
+	return FSubTrackEditor::BuildOutlinerColumnWidget(Params, ColumnName);
+}
+
 TSharedRef<SWidget> FSubAssemblyTrackEditor::BuildAddSubSequenceComboButtonContent(UE::Sequencer::TWeakViewModelPtr<UE::Sequencer::FViewModel> WeakViewModel)
 {
 	using namespace UE::Sequencer;
@@ -263,6 +354,15 @@ void FSubAssemblyTrackEditor::AddNewSequenceToTrack(UMovieSceneSubAssemblyTrack*
 	UCineAssembly* NewAssembly = NewObject<UCineAssembly>();
 	NewAssembly->Initialize();
 
+	// Set the playback range of the new assembly to match the template object (if possible)
+	if (UMovieSceneSequence* Sequence = Cast<UMovieSceneSequence>(TemplateObject))
+	{
+		if (NewAssembly->GetMovieScene() && Sequence->GetMovieScene())
+		{
+			NewAssembly->GetMovieScene()->SetPlaybackRange(Sequence->GetMovieScene()->GetPlaybackRange());
+		}
+	}
+
 	FSubTrackEditor::AddSequenceToTrack(NewAssembly, Track);
 
 	// Find the SubAssemblySection in the Track that was just created (i.e. whose Sequence matches the one just created) and set its Assembly Template to be the input template object
@@ -308,17 +408,23 @@ void FSubAssemblyTrackEditor::OnAssetEnterPressedAsReference(const TArray<FAsset
 
 void FSubAssemblyTrackEditor::OnDroppedAsTemplate(UObject* Asset)
 {
-	UMovieScene* FocusedMovieScene = GetFocusedMovieScene();
-	UMovieSceneSubAssemblyTrack* NewTrack = Cast<UMovieSceneSubAssemblyTrack>(FindOrCreateSubTrack(FocusedMovieScene, nullptr));
+	// Determine the appropriate track type for the dropped asset
+	ESubAssemblyTrackType TrackType = ESubAssemblyTrackType::SubsequenceTrack;
+	if (UMovieSceneSequence* Sequence = Cast<UMovieSceneSequence>(Asset))
+	{
+		TrackType = (Sequence->GetMovieScene()->GetCameraCutTrack()) ? ESubAssemblyTrackType::CinematicShotTrack : ESubAssemblyTrackType::SubsequenceTrack;
+	}
 
+	UMovieSceneSubAssemblyTrack* NewTrack = AddSubAssemblyTrack(TrackType);
 	AddNewSequenceToTrack(NewTrack, Asset);
 }
 
 void FSubAssemblyTrackEditor::OnDroppedAsReference(UMovieSceneSequence* Sequence)
 {
-	UMovieScene* FocusedMovieScene = GetFocusedMovieScene();
-	UMovieSceneSubAssemblyTrack* NewTrack = Cast<UMovieSceneSubAssemblyTrack>(FindOrCreateSubTrack(FocusedMovieScene, nullptr));
+	// Determine the appropriate track type for the dropped sequence
+	ESubAssemblyTrackType TrackType = (Sequence->GetMovieScene()->GetCameraCutTrack()) ? ESubAssemblyTrackType::CinematicShotTrack : ESubAssemblyTrackType::SubsequenceTrack;
 
+	UMovieSceneSubAssemblyTrack* NewTrack = AddSubAssemblyTrack(TrackType);
 	FSubTrackEditor::AddSequenceToTrack(Sequence, NewTrack);
 }
 
diff --git a/Engine/Plugins/VirtualProduction/CinematicAssemblyTools/Source/CineAssemblyToolsEditor/Private/Sequencer/SubAssemblyTrackEditor.h b/Engine/Plugins/VirtualProduction/CinematicAssemblyTools/Source/CineAssemblyToolsEditor/Private/Sequencer/SubAssemblyTrackEditor.h
--- a/Engine/Plugins/VirtualProduction/CinematicAssemblyTools/Source/CineAssemblyToolsEditor/Private/Sequencer/SubAssemblyTrackEditor.h
+++ b/Engine/Plugins/VirtualProduction/CinematicAssemblyTools/Source/CineAssemblyToolsEditor/Private/Sequencer/SubAssemblyTrackEditor.h
@@ -4,6 +4,8 @@
 
 #include "TrackEditors/SubTrackEditor.h"
 
+enum class ESubAssemblyTrackType : uint8;
+
 class FMenuBuilder;
 class UCineAssembly;
 class UMovieSceneSubTrack;
@@ -23,14 +25,22 @@ public:
 	static TSharedRef<ISequencerTrackEditor> CreateTrackEditor(TSharedRef<ISequencer> OwningSequencer);
 
 	//~ Begin ISequencerTrackEditor interface
-	virtual TSubclassOf<UMovieSceneSubTrack> GetSubTrackClass() const override;
-	virtual bool SupportsSequence(UMovieSceneSequence* InSequence) const override;
-	virtual bool HandleAssetAdded(UObject* Asset, const FGuid& TargetObjectGuid) override;
+	virtual void BuildAddTrackMenu(FMenuBuilder& MenuBuilder) override;
 	virtual TSharedPtr<SWidget> BuildOutlinerEditWidget(const FGuid& ObjectBinding, UMovieSceneTrack* Track, const FBuildEditWidgetParams& Params) override;
+	virtual TSharedPtr<SWidget> BuildOutlinerColumnWidget(const FBuildColumnWidgetParams& Params, const FName& ColumnName) override;
+	virtual bool HandleAssetAdded(UObject* Asset, const FGuid& TargetObjectGuid) override;
 	virtual TSharedRef<ISequencerSection> MakeSectionInterface(UMovieSceneSection& SectionObject, UMovieSceneTrack& Track, FGuid ObjectBinding) override;
+	virtual bool SupportsSequence(UMovieSceneSequence* InSequence) const override;
 	//~ End ISequencerTrackEditor interface
 
+	//~ Begin FSubTrackEditor interface
+	virtual TSubclassOf<UMovieSceneSubTrack> GetSubTrackClass() const override;
+	//~ End FSubTrackEditor interface
+
 private:
+	/** Adds a new SubAssembly Track to the focused movie scene, and sets it Track Type appropriately */
+	UMovieSceneSubAssemblyTrack* AddSubAssemblyTrack(ESubAssemblyTrackType TrackType);
+
 	/** Build the menu to display when clicking the "+" button on a SubAssembly Track */
 	TSharedRef<SWidget> BuildAddSubSequenceComboButtonContent(UE::Sequencer::TWeakViewModelPtr<UE::Sequencer::FViewModel> WeakViewModel);
 
