ShotManagement: Fix folder name validation to allow for token characters that may otherwise be invalid.
#rnx

#Synced-CL 50484466

#Synced-CL 50484466

diff --git a/Engine/Plugins/VirtualProduction/CinematicAssemblyTools/Source/CineAssemblyToolsEditor/Private/UI/CineAssembly/SCineAssemblySchemaWindow.cpp b/Engine/Plugins/VirtualProduction/CinematicAssemblyTools/Source/CineAssemblyToolsEditor/Private/UI/CineAssembly/SCineAssemblySchemaWindow.cpp
--- a/Engine/Plugins/VirtualProduction/CinematicAssemblyTools/Source/CineAssemblyToolsEditor/Private/UI/CineAssembly/SCineAssemblySchemaWindow.cpp
+++ b/Engine/Plugins/VirtualProduction/CinematicAssemblyTools/Source/CineAssemblyToolsEditor/Private/UI/CineAssembly/SCineAssemblySchemaWindow.cpp
@@ -1384,11 +1384,14 @@ bool SCineAssemblySchemaWindow::ValidateFolderName(const FText& InText, FText& O
 		return false;
 	}
 
+	// Check for invalid characters
+	FString InvalidCharacters = INVALID_OBJECTNAME_CHARACTERS INVALID_LONGPACKAGE_CHARACTERS;
+
 	// These characters are actually valid, because we want to support naming tokens
-	const FString FolderNameWithoutTokenChars = InText.ToString().Replace(TEXT(":"), TEXT(""));
+	InvalidCharacters = InvalidCharacters.Replace(TEXT("{}"), TEXT(""));
+	InvalidCharacters = InvalidCharacters.Replace(TEXT(":"), TEXT(""));
 
-	// Check for invalid characters
-	if (!AssetViewUtils::IsValidFolderName(FolderNameWithoutTokenChars, OutErrorMessage))
+	if (!FName::IsValidXName(InText.ToString(), InvalidCharacters, &OutErrorMessage))
 	{
 		return false;
 	}
diff --git a/Engine/Plugins/VirtualProduction/CinematicAssemblyTools/Source/CineAssemblyToolsEditor/Private/UI/SFolderHierarchyPanel.cpp b/Engine/Plugins/VirtualProduction/CinematicAssemblyTools/Source/CineAssemblyToolsEditor/Private/UI/SFolderHierarchyPanel.cpp
--- a/Engine/Plugins/VirtualProduction/CinematicAssemblyTools/Source/CineAssemblyToolsEditor/Private/UI/SFolderHierarchyPanel.cpp
+++ b/Engine/Plugins/VirtualProduction/CinematicAssemblyTools/Source/CineAssemblyToolsEditor/Private/UI/SFolderHierarchyPanel.cpp
@@ -697,11 +697,14 @@ bool SFolderHierarchyPanel::IsValidFolderName(const FText& InText, FText& OutErr
 		return false;
 	}
 
+	// Check for invalid characters
+	FString InvalidCharacters = INVALID_OBJECTNAME_CHARACTERS INVALID_LONGPACKAGE_CHARACTERS;
+
 	// These characters are actually valid, because we want to support naming tokens
-	const FString FolderNameWithoutTokenChars = InText.ToString().Replace(TEXT(":"), TEXT(""));
+	InvalidCharacters = InvalidCharacters.Replace(TEXT("{}"), TEXT(""));
+	InvalidCharacters = InvalidCharacters.Replace(TEXT(":"), TEXT(""));
 
-	// Check for invalid characters
-	if (!AssetViewUtils::IsValidFolderName(FolderNameWithoutTokenChars, OutErrorMessage))
+	if (!FName::IsValidXName(InText.ToString(), InvalidCharacters, &OutErrorMessage))
 	{
 		return false;
 	}
