ShotManagement: Update how CineAssembly objects initialize and create SubAssemblies.
- SubAssemblies are now created as Transient objects that are configurable, just like the outer assembly.
- When the CineAssembly Factory creates a new package for a configured Assembly, it will also recursively create packages for all of its SubAssemblies.
- When initializing a SubAssembly using an existing Assembly or Sequence asset, the SubSequences from the existing asset are also duplicated as additional SubAssemblies.
ShotManagement: Deprecate a few subsequence related properties of the UCineAssemblySchema now that there is support for SubAssembly Tracks in the Schema's TemplateSequence.
- Fix up existing Schema assets in PostLoad by creating new SubAssembly Tracks from the now deprecated subsequence properties.
#rnx

#Synced-CL 50484466

diff --git a/Engine/Plugins/VirtualProduction/CinematicAssemblyTools/Source/CineAssemblyTools/Private/CineAssembly.cpp b/Engine/Plugins/VirtualProduction/CinematicAssemblyTools/Source/CineAssemblyTools/Private/CineAssembly.cpp
--- a/Engine/Plugins/VirtualProduction/CinematicAssemblyTools/Source/CineAssemblyTools/Private/CineAssembly.cpp
+++ b/Engine/Plugins/VirtualProduction/CinematicAssemblyTools/Source/CineAssemblyTools/Private/CineAssembly.cpp
@@ -2,6 +2,7 @@
 
 #include "CineAssembly.h"
 
+#include "Algo/Contains.h"
 #include "AssetRegistry/AssetRegistryModule.h"
 #include "CineAssemblyNamingTokens.h"
 #include "HAL/FileManager.h"
@@ -147,7 +148,6 @@ ETrackSupport UCineAssembly::IsTrackSupportedImpl(TSubclassOf<class UMovieSceneT
 }
 #endif //WITH_EDITOR
 
-
 void UCineAssembly::Initialize()
 {
 	Super::Initialize();
@@ -163,22 +163,27 @@ void UCineAssembly::Initialize()
 #endif // WITH_EDITOR
 }
 
-void UCineAssembly::InitializeFromTemplate(ULevelSequence* Template)
+void UCineAssembly::InitializeFromSchema(UCineAssemblySchema* InSchema)
 {
-	Initialize();
+	SetSchema(InSchema);
 
-	if (Template == nullptr)
+	if (InSchema)
 	{
-		return;
+		InitializeFromTemplate(InSchema->TemplateSequence);
+		ConvertSubAssemblyTracks();
 	}
+}
 
-	// If the template is a CineAssembly, set the schema of this assembly to match
-	if (UCineAssembly* TemplateAssembly = Cast<UCineAssembly>(Template))
+void UCineAssembly::InitializeFromTemplate(ULevelSequence* Template)
+{
+	if (!MovieScene)
 	{
-		if (const UCineAssemblySchema* TemplateSchema = TemplateAssembly->GetSchema())
-		{
-			SetSchema(const_cast<UCineAssemblySchema*>(TemplateSchema));
-		}
+		Initialize();
+	}
+
+	if (Template == nullptr)
+	{
+		return;
 	}
 
 	if (UMovieScene* TemplateMovieScene = Template->GetMovieScene())
@@ -215,163 +220,8 @@ void UCineAssembly::InitializeFromTemplate(ULevelSequence* Template)
 				ThisBindingReferencesRef.AddBinding(Reference.ID, MoveTemp(NewLocator), Reference.ResolveFlags, CustomBindingDuplicate);
 			}
 		}
-
-#if WITH_EDITOR
-		// Replace any SubAssemblyTracks we find with either a MovieSceneSubTrack or MovieSceneCinematicShotTrack, depending on the TrackType
-		while (UMovieSceneSubAssemblyTrack* SubAssemblyTrack = MovieScene->FindTrack<UMovieSceneSubAssemblyTrack>())
-		{
-			// Add a new track to the MovieScene. The SubAssemblyTrack will be removed from the MovieScene at the bottom of this loop.
-			UClass* TrackClass = SubAssemblyTrack->IsSubsequenceTrack() ? UMovieSceneSubTrack::StaticClass() : UMovieSceneCinematicShotTrack::StaticClass();
-			UMovieSceneNameableTrack* NewTrack = Cast<UMovieSceneNameableTrack>(MovieScene->AddTrack(TrackClass));
-
-			NewTrack->SetDisplayName(SubAssemblyTrack->GetDisplayName());
-
-			for (UMovieSceneSection* Section : SubAssemblyTrack->GetAllSections())
-			{
-				if (UMovieSceneSubAssemblySection* SubAssemblySection = Cast<UMovieSceneSubAssemblySection>(Section))
-				{
-					// Create a new section, with the relevant properties copied from the SubAssemblySection, and add it to the new track
-					UClass* SectionClass = SubAssemblyTrack->IsSubsequenceTrack() ? UMovieSceneSubSection::StaticClass() : UMovieSceneCinematicShotSection::StaticClass();
-					UMovieSceneSubSection* NewSubSection = NewObject<UMovieSceneSubSection>(NewTrack, SectionClass, NAME_None, SubAssemblySection->GetFlags(), SubAssemblySection);
-
-					NewTrack->AddSection(*NewSubSection);
-
-					// Remove the tint from the section that was used in the Schema editor to differentiate between template and reference sections
-					if (NewSubSection->GetColorTint() == UMovieSceneSubAssemblySection::TemplateSectionColorTint)
-					{
-						NewSubSection->SetColorTint(UMovieSceneSubAssemblySection::ReferenceSectionColorTint);
-					}
-
-					if (SubAssemblySection->IsTemplateSection())
-					{
-						// Create a new SubAssembly using the assembly template specified by the SubAssemblySection
-						UCineAssembly* NewSubAssembly = CreateSubAssemblyFromTemplate(SubAssemblySection->NewSequenceName.ToString(), SubAssemblySection->GetAssemblyTemplate());
-						if (NewSubAssembly)
-						{
-							NewSubSection->SetSequence(NewSubAssembly);
-							SubAssemblies.Add(NewSubSection);
-						}
-					}
-				}
-			}
-
-			// Now remove the SubAssemblyTrack from the MovieScene. 
-			MovieScene->RemoveTrack(*SubAssemblyTrack);
-		}
-
-		// Recursively gather all of the folders in the movie scene
-		TArray<UMovieSceneFolder*> AllFolders;
-		for (UMovieSceneFolder* RootFolder : MovieScene->GetRootFolders())
-		{
-			AllFolders.Add(RootFolder);
-			GetAllMovieSceneFoldersRecursive(RootFolder, AllFolders);
-		}
-
-		// Resolve the tokens in each folder name
-		for (UMovieSceneFolder* Folder : AllFolders)
-		{
-			const FText ResolvedFolderName = UCineAssemblyNamingTokens::GetResolvedText(Folder->GetFolderName().ToString(), this);
-			Folder->SetFolderName(*ResolvedFolderName.ToString());
-		}
-
-		// Resolve the tokens in each track name
-		for (UMovieSceneTrack* Track : MovieScene->GetTracks())
-		{
-			if (UMovieSceneNameableTrack* NameableTrack = Cast<UMovieSceneNameableTrack>(Track))
-			{
-				const FText ResolvedTrackName = UCineAssemblyNamingTokens::GetResolvedText(NameableTrack->GetDisplayName().ToString(), this);
-				NameableTrack->SetDisplayName(ResolvedTrackName);
-			}
-		}
-
-		const TArray<FMovieSceneBinding>& Bindings = Cast<const UMovieScene>(MovieScene)->GetBindings();
-		for (const FMovieSceneBinding& Binding : Bindings)
-		{
-			TArray<UMovieSceneTrack*> Tracks = MovieScene->FindTracks(UMovieSceneNameableTrack::StaticClass(), Binding.GetObjectGuid(), NAME_None);
-			for (UMovieSceneTrack* Track : Tracks)
-			{
-				if (UMovieSceneNameableTrack* NameableTrack = Cast<UMovieSceneNameableTrack>(Track))
-				{
-					const FText ResolvedTrackName = UCineAssemblyNamingTokens::GetResolvedText(NameableTrack->GetDisplayName().ToString(), this);
-					NameableTrack->SetDisplayName(ResolvedTrackName);
-				}
-			}
-		}
-#endif // WITH_EDITOR
-	}
-}
-
-#if WITH_EDITOR
-UCineAssembly* UCineAssembly::CreateSubAssemblyFromTemplate(const FString& SubAssemblyName, UObject* TemplateObject)
-{
-	// The new SubAssembly begins are a transient object so that it can be configured before it is named and given a package.
-	// This allows us to set its metadata first before trying to resolve any tokens which might appear in the SubAssemblyName.
-	UCineAssembly* SubAssembly = NewObject<UCineAssembly>(GetTransientPackage(), NAME_None, RF_Public | RF_Standalone | RF_Transactional);
-
-	// Initialize the new SubAssembly using the TemplateObject
-	if (UCineAssemblySchema* SchemaTemplate = Cast<UCineAssemblySchema>(TemplateObject))
-	{
-		SubAssembly->Initialize();
-		SubAssembly->SetSchema(SchemaTemplate);
-	}
-	else if (ULevelSequence* SequenceTemplate = Cast<ULevelSequence>(TemplateObject))
-	{
-		SubAssembly->InitializeFromTemplate(SequenceTemplate);
-	}
-	else
-	{
-		SubAssembly->Initialize();
-	}
-
-	SubAssembly->Level = Level;
-	SubAssembly->ParentAssembly = this;
-	SubAssembly->Production = Production;
-	SubAssembly->ProductionName = ProductionName;
-
-	// Evaluate the name of the assembly
-	SubAssembly->AssemblyName.Template = SubAssemblyName;
-	SubAssembly->AssemblyName.Resolved = UCineAssemblyNamingTokens::GetResolvedText(SubAssembly->AssemblyName.Template, SubAssembly);
-
-	// If the assembly name is empty or invalid, assign it a valid default name
-	if (SubAssembly->AssemblyName.Resolved.IsEmpty() || FName(*SubAssembly->AssemblyName.Resolved.ToString()).IsNone())
-	{
-		SubAssembly->AssemblyName.Resolved = LOCTEXT("NewCineAssemblyName", "NewCineAssembly");
-		SubAssembly->AssemblyName.Template = SubAssembly->AssemblyName.Resolved.ToString();
-	}
-
-	FString AssemblyPath;
-	FString AssemblyRootFolder;
-	GetAssetPathAndRootFolder(AssemblyPath, AssemblyRootFolder);
-
-	IAssetTools& AssetTools = FModuleManager::GetModuleChecked<FAssetToolsModule>("AssetTools").Get();
-	FString UniquePackageName;
-	FString UniqueAssetName;
-	AssetTools.CreateUniqueAssetName(AssemblyRootFolder / SubAssembly->AssemblyName.Resolved.ToString(), TEXT(""), UniquePackageName, UniqueAssetName);
-
-	// Now, we can create a real package for it and rename it
-	UPackage* Package = CreatePackage(*UniquePackageName);
-	SubAssembly->Rename(*UniqueAssetName, Package);
-
-	SubAssembly->CreateSubAssemblies();
-
-	// Notify the asset registry about the new assembly
-	FAssetRegistryModule::AssetCreated(SubAssembly);
-
-	// Mark the package dirty
-	Package->MarkPackageDirty();
-
-	return SubAssembly;
-}
-
-void UCineAssembly::GetAllMovieSceneFoldersRecursive(UMovieSceneFolder* RootFolder, TArray<UMovieSceneFolder*>& OutFolders)
-{
-	for (UMovieSceneFolder* ChildFolder : RootFolder->GetChildFolders())
-	{
-		OutFolders.Add(ChildFolder);
-		GetAllMovieSceneFoldersRecursive(ChildFolder, OutFolders);
 	}
 }
-#endif // WITH_EDITOR
 
 FGuid UCineAssembly::GetAssemblyGuid() const
 {
@@ -453,18 +303,6 @@ void UCineAssembly::ChangeSchema(UCineAssemblySchema* InSchema)
 		}
 	}
 
-	// Reset the list of SubAssembly names to create from the Schema
-	SubAssemblyNames.Reset();
-	if (BaseSchema)
-	{
-		for (const FString& SubAssemblyName : BaseSchema->SubsequencesToCreate)
-		{
-			FTemplateString NewTemplateName;
-			NewTemplateName.Template = SubAssemblyName;
-			SubAssemblyNames.Add(NewTemplateName);
-		}
-	}
-
 	// Reset the list of folders names to create from the Schema
 	DefaultFolderNames.Reset();
 	if (BaseSchema)
@@ -487,11 +325,202 @@ void UCineAssembly::ChangeSchema(UCineAssemblySchema* InSchema)
 #if WITH_EDITOR
 void UCineAssembly::CreateSubAssemblies()
 {
-	if (!BaseSchema)
+	// Deprecated
+}
+#endif // WITH_EDITOR
+
+UCineAssembly* UCineAssembly::CreateSubAssemblyFromTemplate(UObject* TemplateObject)
+{
+	// The new SubAssembly begins are a transient object so that it can be configured before it is named and given a package.
+	// This allows us to set its metadata first before trying to resolve any tokens which might appear in the SubAssemblyName.
+	UCineAssembly* SubAssembly = NewObject<UCineAssembly>(GetTransientPackage());
+
+	if (UCineAssemblySchema* SchemaTemplate = Cast<UCineAssemblySchema>(TemplateObject))
 	{
-		return;
+		SubAssembly->InitializeFromSchema(SchemaTemplate);
+	}
+	else if (UCineAssembly* AssemblyTemplate = Cast<UCineAssembly>(TemplateObject))
+	{
+		SubAssembly->SetSchema(const_cast<UCineAssemblySchema*>(AssemblyTemplate->GetSchema()));
+		SubAssembly->InitializeFromTemplate(AssemblyTemplate);
+		SubAssembly->DuplicateSubSequences(AssemblyTemplate);
+	}
+	else if (ULevelSequence* SequenceTemplate = Cast<ULevelSequence>(TemplateObject))
+	{
+		SubAssembly->InitializeFromTemplate(SequenceTemplate);
+		SubAssembly->DuplicateSubSequences(SequenceTemplate);
+	}
+	else
+	{
+		SubAssembly->Initialize();
 	}
 
+	return SubAssembly;
+}
+
+void UCineAssembly::ConvertSubAssemblyTracks()
+{
+	// Replace any SubAssemblyTracks we find with either a MovieSceneSubTrack or MovieSceneCinematicShotTrack, depending on the TrackType
+	while (UMovieSceneSubAssemblyTrack* SubAssemblyTrack = MovieScene->FindTrack<UMovieSceneSubAssemblyTrack>())
+	{
+		// Add a new track to the MovieScene. The SubAssemblyTrack will be removed from the MovieScene at the bottom of this loop.
+		UClass* TrackClass = SubAssemblyTrack->IsSubsequenceTrack() ? UMovieSceneSubTrack::StaticClass() : UMovieSceneCinematicShotTrack::StaticClass();
+		UMovieSceneNameableTrack* NewTrack = Cast<UMovieSceneNameableTrack>(MovieScene->AddTrack(TrackClass));
+
+#if WITH_EDITOR
+		NewTrack->SetDisplayName(SubAssemblyTrack->GetDisplayName());
+#endif // WITH_EDITOR
+
+		for (UMovieSceneSection* Section : SubAssemblyTrack->GetAllSections())
+		{
+			if (UMovieSceneSubAssemblySection* SubAssemblySection = Cast<UMovieSceneSubAssemblySection>(Section))
+			{
+				// Create a new section, with the relevant properties copied from the SubAssemblySection, and add it to the new track
+				UClass* SectionClass = SubAssemblyTrack->IsSubsequenceTrack() ? UMovieSceneSubSection::StaticClass() : UMovieSceneCinematicShotSection::StaticClass();
+				UMovieSceneSubSection* NewSubSection = NewObject<UMovieSceneSubSection>(NewTrack, SectionClass, NAME_None, SubAssemblySection->GetFlags(), SubAssemblySection);
+
+				NewTrack->AddSection(*NewSubSection);
+
+				// Remove the tint from the section that was used in the Schema editor to differentiate between template and reference sections
+				if (NewSubSection->GetColorTint() == UMovieSceneSubAssemblySection::TemplateSectionColorTint)
+				{
+					NewSubSection->SetColorTint(UMovieSceneSubAssemblySection::ReferenceSectionColorTint);
+				}
+
+				if (SubAssemblySection->IsTemplateSection())
+				{
+					UCineAssembly* NewSubAssembly = CreateSubAssemblyFromTemplate(SubAssemblySection->GetAssemblyTemplate());
+
+					// Set the name and path of the new SubAssembly from the data stored in the SubAssembly section
+					NewSubAssembly->AssemblyName.Template = SubAssemblySection->GetSequenceName().ToString();
+					NewSubAssembly->RelativeAssemblyPath.Template = SubAssemblySection->GetSequencePath();
+
+					if (!Algo::ContainsBy(DefaultFolderNames, NewSubAssembly->RelativeAssemblyPath.Template, &FTemplateString::Template))
+					{
+						DefaultFolderNames.Add(NewSubAssembly->RelativeAssemblyPath);
+					}
+
+					NewSubSection->SetSequence(NewSubAssembly);
+					SubAssemblies.Add(NewSubSection);
+				}
+			}
+		}
+
+		// Now remove the SubAssemblyTrack from the MovieScene. 
+		MovieScene->RemoveTrack(*SubAssemblyTrack);
+	}
+}
+
+void UCineAssembly::DuplicateSubSequences(ULevelSequence* SequenceTemplate)
+{
+	for (UMovieSceneTrack* Track : MovieScene->GetTracks())
+	{
+		if (Track->IsA<UMovieSceneSubTrack>() || Track->IsA<UMovieSceneCinematicShotTrack>())
+		{
+			for (UMovieSceneSection* Section : Track->GetAllSections())
+			{
+				if (UMovieSceneSubSection* SubSection = Cast<UMovieSceneSubSection>(Section))
+				{
+					ULevelSequence* SubSequence = Cast<ULevelSequence>(SubSection->GetSequence());
+					if (!SubSequence)
+					{
+						continue;
+					}
+
+					// If the SubSequence in this section is already a transient Assembly, skip duplicating it
+					if (SubSequence && SubSequence->IsA<UCineAssembly>() && SubSequence->GetPackage() == GetTransientPackage())
+					{
+						continue;
+					}
+
+					// Create a new SubAssembly using the existing SubSequence as a template
+					UCineAssembly* DuplicateSubAssembly = CreateSubAssemblyFromTemplate(SubSequence);
+
+					if (UCineAssembly* SubAssembly = Cast<UCineAssembly>(SubSequence))
+					{
+						DuplicateSubAssembly->AssemblyName.Template = SubAssembly->AssemblyName.Template;
+					}
+					else
+					{
+#if WITH_EDITOR
+						DuplicateSubAssembly->AssemblyName.Template = SubSequence->GetDisplayName().ToString();
+#else
+						DuplicateSubAssembly->AssemblyName.Template = SubSequence->GetName();
+#endif // WITH_EDITOR
+					}
+
+					// When the transient SubAssembly is converted into an asset, the relative path will tell the factory where to create the package relative to the root folder
+					DuplicateSubAssembly->RelativeAssemblyPath.Template = GetRelativePath(SubSequence, SequenceTemplate);
+
+					// Replace the sequence in the section with the duplicate assembly
+					SubSection->SetSequence(DuplicateSubAssembly);
+					SubAssemblies.Add(SubSection);
+				}
+			}
+		}
+	}
+}
+
+FString UCineAssembly::GetRelativePath(ULevelSequence* SubSequence, ULevelSequence* OuterSequence)
+{
+	// If the SubSequence is an Assembly, the path we want is its content root folder. 
+	// Otherwise, we just want the actual path of the asset
+	FString SubSequencePath;
+	if (UCineAssembly* SubAssembly = Cast<UCineAssembly>(SubSequence))
+	{
+		FString AssetPath; // Unused
+		SubAssembly->GetAssetPathAndRootFolder(AssetPath, SubSequencePath);
+	}
+	else
+	{
+		SubSequencePath = FPaths::GetPath(SubSequence->GetPathName());
+	}
+
+	// If there is no outer sequence to make the path relative to, early-out
+	if (!OuterSequence)
+	{
+		return SubSequencePath;
+	}
+
+	// If the outer sequence is an Assembly, the we want is its content root folder.
+	// Otherwise, we just want the actual path of the asset
+	FString OuterSequencePath;
+	if (UCineAssembly* OuterAssembly = Cast<UCineAssembly>(OuterSequence))
+	{
+		FString OuterAssemblyAssetPath; // Unused
+		OuterAssembly->GetAssetPathAndRootFolder(OuterAssemblyAssetPath, OuterSequencePath);
+	}
+	else
+	{
+		OuterSequencePath = FPaths::GetPath(OuterSequence->GetPathName());
+	}
+
+	// Only attempt to find a relative path if the SubSequence resides within the content tree of the outer sequence
+	if (FPaths::IsUnderDirectory(SubSequencePath, OuterSequencePath))
+	{
+		// Make the path relative to the outer sequence
+		FPaths::MakePathRelativeTo(SubSequencePath, *OuterSequencePath.AppendChar('/'));
+
+		// If the outer sequence is an Assembly, check if the relative subsequence path matches any of the outer assembly's resolved default folders.
+		// If a match is found, return that folder's template name (potentially tokenized)
+		if (UCineAssembly* OuterAssembly = Cast<UCineAssembly>(OuterSequence))
+		{
+			for (const FTemplateString& DefaultPath : OuterAssembly->DefaultFolderNames)
+			{
+				if (DefaultPath.Resolved.ToString() == SubSequencePath)
+				{
+					return DefaultPath.Template;
+				}
+			}
+		}
+	}
+
+	return SubSequencePath;
+}
+
+#if WITH_EDITOR
+void UCineAssembly::CreateSubFolders()
+{
 	// Get the path where the top-level assembly will be created so we can create other assets relative to it
 	FString AssemblyPath;
 	FString AssemblyRootFolder;
@@ -523,87 +552,59 @@ void UCineAssembly::CreateSubAssemblies()
 			}
 		}
 	}
+}
 
-	// Early-out if the subassemblies have already been created (as would be the case if this was a duplicated assembly)
-	if (!SubAssemblies.IsEmpty())
+void UCineAssembly::ResolveMovieSceneTokens()
+{
+	// Recursively gather all of the folders in the movie scene
+	TArray<UMovieSceneFolder*> AllFolders;
+	for (UMovieSceneFolder* RootFolder : MovieScene->GetRootFolders())
 	{
-		return;
+		AllFolders.Add(RootFolder);
+		GetAllMovieSceneFoldersRecursive(RootFolder, AllFolders);
 	}
 
-	IAssetTools& AssetTools = FModuleManager::GetModuleChecked<FAssetToolsModule>("AssetTools").Get();
-
-	// Create a new CineAssembly for each subsequence, set its playback range to match the parent sequence, and add it to the subsequence track
-	for (FTemplateString& SubAssemblyName : SubAssemblyNames)
+	// Resolve the tokens in each folder name
+	for (UMovieSceneFolder* Folder : AllFolders)
 	{
-		// Resolve any tokens found in the name of the SubAssembly before attempting to create it 
-		SubAssemblyName.Resolved = UCineAssemblyNamingTokens::GetResolvedText(SubAssemblyName.Template, this);
-
-		const FString SubAssemblyFilename = FPaths::GetBaseFilename(SubAssemblyName.Resolved.ToString());
+		const FText ResolvedFolderName = UCineAssemblyNamingTokens::GetResolvedText(Folder->GetFolderName().ToString(), this);
+		Folder->SetFolderName(*ResolvedFolderName.ToString());
+	}
 
-		if (SubAssemblyFilename.IsEmpty())
+	// Resolve the tokens in each track name
+	for (UMovieSceneTrack* Track : MovieScene->GetTracks())
+	{
+		if (UMovieSceneNameableTrack* NameableTrack = Cast<UMovieSceneNameableTrack>(Track))
 		{
-			continue;
+			const FText ResolvedTrackName = UCineAssemblyNamingTokens::GetResolvedText(NameableTrack->GetDisplayName().ToString(), this);
+			NameableTrack->SetDisplayName(ResolvedTrackName);
 		}
+	}
 
-		// Add a subsequence track to the assembly's sequence
-		UMovieSceneSubTrack* SubTrack = Cast<UMovieSceneSubTrack>(MovieScene->AddTrack(UMovieSceneSubTrack::StaticClass()));
-		SubTrack->SetDisplayName(FText::FromString(SubAssemblyFilename));
-
-		// Before creating each subassembly, sanity check that each one will actually have a unique name (in case there are duplicates in the schema description)
-		FString UniquePackageName;
-		FString UniqueAssetName;
-		AssetTools.CreateUniqueAssetName(AssemblyRootFolder / SubAssemblyName.Resolved.ToString(), TEXT(""), UniquePackageName, UniqueAssetName);
-
-		const FString SubAssemblyPath = FPaths::GetPath(UniquePackageName);
-
-		UPackage* SubAssemblyPackage = CreatePackage(*UniquePackageName);
-
-		const EObjectFlags Flags = RF_Public | RF_Standalone | RF_Transactional;
-		UCineAssembly* SubAssembly = NewObject<UCineAssembly>(SubAssemblyPackage, UCineAssembly::StaticClass(), FName(*UniqueAssetName), Flags);
-
-		if (SubAssembly)
+	const TArray<FMovieSceneBinding>& Bindings = Cast<const UMovieScene>(MovieScene)->GetBindings();
+	for (const FMovieSceneBinding& Binding : Bindings)
+	{
+		TArray<UMovieSceneTrack*> Tracks = MovieScene->FindTracks(UMovieSceneNameableTrack::StaticClass(), Binding.GetObjectGuid(), NAME_None);
+		for (UMovieSceneTrack* Track : Tracks)
 		{
-			ULevelSequence* Template = nullptr;
-			// Check to see if there is a supplied templated for this Subsequence.
-			if (FSoftObjectPath* TemplatePath = BaseSchema->SubsequenceTemplates.Find(SubAssemblyName.Template))
-			{
-				Template = Cast<ULevelSequence>(TemplatePath->TryLoad());
-			}
-
-			if (Template != nullptr)
-			{
-				SubAssembly->InitializeFromTemplate(Template);
-			}
-			else
+			if (UMovieSceneNameableTrack* NameableTrack = Cast<UMovieSceneNameableTrack>(Track))
 			{
-				SubAssembly->Initialize();
+				const FText ResolvedTrackName = UCineAssemblyNamingTokens::GetResolvedText(NameableTrack->GetDisplayName().ToString(), this);
+				NameableTrack->SetDisplayName(ResolvedTrackName);
 			}
-
-			const TRange<FFrameNumber> PlaybackRange = MovieScene->GetPlaybackRange();
-			SubAssembly->GetMovieScene()->SetPlaybackRange(PlaybackRange);
-
-			SubAssembly->Level = Level;
-			SubAssembly->ParentAssembly = this;
-			SubAssembly->Production = Production;
-			SubAssembly->ProductionName = ProductionName;
-
-			ULevelSequenceShotMetaDataLibrary::SetIsSubSequence(SubAssembly, true);
-
-			const FFrameNumber StartFrame = PlaybackRange.GetLowerBoundValue();
-			const int32 Duration = PlaybackRange.Size<FFrameNumber>().Value;
-
-			UMovieSceneSubSection* SubSection = SubTrack->AddSequence(SubAssembly, StartFrame, Duration);
-			SubAssemblies.Add(SubSection);
-
-			// Notify the asset registry about the new subassembly
-			FAssetRegistryModule::AssetCreated(SubAssembly);
-
-			// Mark the package dirty
-			SubAssemblyPackage->MarkPackageDirty();
 		}
 	}
 }
-#endif
+
+void UCineAssembly::GetAllMovieSceneFoldersRecursive(UMovieSceneFolder* RootFolder, TArray<UMovieSceneFolder*>& OutFolders)
+{
+	for (UMovieSceneFolder* ChildFolder : RootFolder->GetChildFolders())
+	{
+		OutFolders.Add(ChildFolder);
+		GetAllMovieSceneFoldersRecursive(ChildFolder, OutFolders);
+	}
+}
+#endif // WITH_EDITOR
 
 void UCineAssembly::GetAssetRegistryTags(FAssetRegistryTagsContext Context) const
 {
@@ -650,7 +651,7 @@ void UCineAssembly::PostEditChangeProperty(struct FPropertyChangedEvent& Propert
 	}
 }
 
-#endif
+#endif // WITH_EDITOR
 
 void UCineAssembly::Serialize(FArchive& Ar)
 {
diff --git a/Engine/Plugins/VirtualProduction/CinematicAssemblyTools/Source/CineAssemblyTools/Private/CineAssemblySchema.cpp b/Engine/Plugins/VirtualProduction/CinematicAssemblyTools/Source/CineAssemblyTools/Private/CineAssemblySchema.cpp
--- a/Engine/Plugins/VirtualProduction/CinematicAssemblyTools/Source/CineAssemblyTools/Private/CineAssemblySchema.cpp
+++ b/Engine/Plugins/VirtualProduction/CinematicAssemblyTools/Source/CineAssemblyTools/Private/CineAssemblySchema.cpp
@@ -7,6 +7,8 @@
 #include "CineAssembly.h"
 #include "Modules/ModuleManager.h"
 #include "MovieScene.h"
+#include "MovieSceneSubAssemblySection.h"
+#include "MovieSceneSubAssemblyTrack.h"
 
 #if WITH_EDITOR
 #include "AssetToolsModule.h"
@@ -69,6 +71,67 @@ void UCineAssemblySchema::PostLoad()
 	{
 		SchemaGuid = FGuid::NewGuid();
 	}
+
+#if WITH_EDITOR
+	PRAGMA_DISABLE_DEPRECATION_WARNINGS
+	// Fix-up old schema assets that were created before support was added for the fully editable template sequence
+	// If a top-level Template was previously set, use it to initialize the TemplateSequence
+	if (TemplateSequence && TemplateSequence->GetMovieScene())
+	{
+		if (ULevelSequence* OldTemplate = Cast<ULevelSequence>(Template_DEPRECATED.TryLoad()))
+		{
+			TemplateSequence->InitializeFromTemplate(OldTemplate);
+		}
+		Template_DEPRECATED.Reset();
+
+		UMovieScene* TemplateSequenceMovieScene = TemplateSequence->GetMovieScene();
+		const TRange<FFrameNumber> TemplateSequencePlaybackRange = TemplateSequenceMovieScene->GetPlaybackRange();
+
+		// If any named subsequences were previously specified in SubsequencesToCreate, add new SubAssembly tracks/sections to the TemplateSequence instead
+		for (const FString& SubAssemblyName : SubsequencesToCreate_DEPRECATED)
+		{
+			// Add a new SubAssembly Track to the TemplateSequence
+			UMovieSceneSubAssemblyTrack* SubAssemblyTrack = Cast<UMovieSceneSubAssemblyTrack>(TemplateSequenceMovieScene->AddTrack(UMovieSceneSubAssemblyTrack::StaticClass()));
+			SubAssemblyTrack->SetTrackType(ESubAssemblyTrackType::SubsequenceTrack);
+
+			SubAssemblyTrack->SetLocalEvalDisabled(true);
+
+			// Add a new SubAssembly Section to the SubAssembly Track
+			const FFrameNumber StartFrame = TemplateSequencePlaybackRange.GetLowerBoundValue();
+			const int32 Duration = TemplateSequencePlaybackRange.Size<FFrameNumber>().Value;
+
+			UMovieSceneSubAssemblySection* SubAssemblySection = Cast<UMovieSceneSubAssemblySection>(SubAssemblyTrack->AddSequence(nullptr, StartFrame, Duration));
+
+			// Set the Assembly Template of the new section 
+			if (FSoftObjectPath* TemplatePath = SubsequenceTemplates_DEPRECATED.Find(SubAssemblyName))
+			{
+				ULevelSequence* SubAssemblyTemplate = Cast<ULevelSequence>(TemplatePath->TryLoad());
+				SubAssemblySection->SetAssemblyTemplate(SubAssemblyTemplate);
+			}
+			else
+			{
+				SubAssemblySection->SetAssemblyTemplate(nullptr);
+			}
+
+			// Set the name and path of the new section based on the old name of this SubsequenceToCreate property
+			FString NewSequenceName = FPaths::GetPathLeaf(SubAssemblyName);
+			FString NewSequencePath = FPaths::GetPath(SubAssemblyName);
+
+			// Replace any occurrences of the "{assembly}" token in the NewSequenceName with the "{parent}" token
+			// Assemblies cannot be named using the "{assembly}" token because it would be self-referencing
+			// The prior intent of using the "{assembly}" token in the name of subsequences was to name them after their parent sequence
+			// Updating the tokens here correctly expresses that intent and ensures the subsequences will be named correctly. 
+			NewSequenceName = NewSequenceName.Replace(TEXT("{cat:assembly}"), TEXT("{cat:parent}"));
+			NewSequenceName = NewSequenceName.Replace(TEXT("{assembly}"), TEXT("{parent}"));
+
+			SubAssemblySection->SetSequenceName(FText::FromString(NewSequenceName));
+			SubAssemblySection->SetSequencePath(NewSequencePath);
+		}
+
+		SubsequencesToCreate_DEPRECATED.Empty();
+	}
+	PRAGMA_ENABLE_DEPRECATION_WARNINGS
+#endif // WITH_EDITOR
 }
 
 FString UCineAssemblySchema::GetDefaultAssemblyName() const
diff --git a/Engine/Plugins/VirtualProduction/CinematicAssemblyTools/Source/CineAssemblyTools/Private/MovieSceneSubAssemblySection.cpp b/Engine/Plugins/VirtualProduction/CinematicAssemblyTools/Source/CineAssemblyTools/Private/MovieSceneSubAssemblySection.cpp
--- a/Engine/Plugins/VirtualProduction/CinematicAssemblyTools/Source/CineAssemblyTools/Private/MovieSceneSubAssemblySection.cpp
+++ b/Engine/Plugins/VirtualProduction/CinematicAssemblyTools/Source/CineAssemblyTools/Private/MovieSceneSubAssemblySection.cpp
@@ -45,7 +45,6 @@ UObject* UMovieSceneSubAssemblySection::GetAssemblyTemplate() const
 	return AssemblyTemplate;
 }
 
-#if WITH_EDITOR
 void UMovieSceneSubAssemblySection::SetAssemblyTemplate(UObject* TemplateObject)
 {
 	SectionType = ESubAssemblySectionType::Template;
@@ -75,7 +74,11 @@ void UMovieSceneSubAssemblySection::SetAssemblyTemplate(UObject* TemplateObject)
 	}
 	else if (ULevelSequence* SequenceTemplate = Cast<ULevelSequence>(TemplateObject))
 	{
+#if WITH_EDITOR
 		NewSequenceName = SequenceTemplate->GetDisplayName();
+#else
+		NewSequenceName = FText::FromString(SequenceTemplate->GetName());
+#endif // WITH_EDITOR
 	}
 	else
 	{
@@ -83,6 +86,27 @@ void UMovieSceneSubAssemblySection::SetAssemblyTemplate(UObject* TemplateObject)
 	}
 }
 
+FText UMovieSceneSubAssemblySection::GetSequenceName() const
+{
+	return NewSequenceName;
+}
+
+void UMovieSceneSubAssemblySection::SetSequenceName(const FText& InName)
+{
+	NewSequenceName = InName;
+}
+
+FString UMovieSceneSubAssemblySection::GetSequencePath() const
+{
+	return NewSequencePath;
+}
+
+void UMovieSceneSubAssemblySection::SetSequencePath(const FString& InPath)
+{
+	NewSequencePath = InPath;
+}
+
+#if WITH_EDITOR
 void UMovieSceneSubAssemblySection::PostEditChangeProperty(FPropertyChangedEvent& PropertyChangedEvent)
 {
 	Super::PostEditChangeProperty(PropertyChangedEvent);
diff --git a/Engine/Plugins/VirtualProduction/CinematicAssemblyTools/Source/CineAssemblyTools/Public/CineAssembly.h b/Engine/Plugins/VirtualProduction/CinematicAssemblyTools/Source/CineAssemblyTools/Public/CineAssembly.h
--- a/Engine/Plugins/VirtualProduction/CinematicAssemblyTools/Source/CineAssemblyTools/Public/CineAssembly.h
+++ b/Engine/Plugins/VirtualProduction/CinematicAssemblyTools/Source/CineAssemblyTools/Public/CineAssembly.h
@@ -38,6 +38,9 @@ public:
 	CINEASSEMBLYTOOLS_API virtual void Initialize() override;
 	//~ End ULevelSquence Interface
 
+	/** Set the Schema for this Assembly and initialize the MovieScene from the Schema's template sequence */
+	CINEASSEMBLYTOOLS_API void InitializeFromSchema(UCineAssemblySchema* InSchema);
+
 	/** Given a ULevelSequence as a template, initialize this assembly from it. */
 	CINEASSEMBLYTOOLS_API void InitializeFromTemplate(ULevelSequence* Template);
 
@@ -66,7 +69,14 @@ public:
 
 #if WITH_EDITOR
 	/** Creates one or more subsequence assets, parented to this assembly, based on the schema */
+	UE_DEPRECATED(5.8, "SubAssemblies are created automatically when initializing an assembly with a Schema")
 	CINEASSEMBLYTOOLS_API void CreateSubAssemblies();
+
+	/** Creates the content folders specified by the schema */
+	CINEASSEMBLYTOOLS_API void CreateSubFolders();
+
+	/** Resolve any tokens in the names of this Assembly's MovieScene folders and tracks */
+	CINEASSEMBLYTOOLS_API void ResolveMovieSceneTokens();
 #endif
 
 	/** Get the target level associated with this assembly */
@@ -184,6 +194,24 @@ private:
 	/** Sets the base schema for this assembly and re-initializes the metadata inherited from the schema */
 	CINEASSEMBLYTOOLS_API void ChangeSchema(UCineAssemblySchema* InSchema);
 
+	/**
+	 * Convert all SubAssembly Tracks found in this Assembly's MovieScene into SubTracks and CinematicShotTracks.
+	 * A new transient SubAssembly will be created for every template SubAssembly Section that is found.
+	 */
+	void ConvertSubAssemblyTracks();
+
+	/** Creates a transient CineAssembly object using the input object as a template */
+	UCineAssembly* CreateSubAssemblyFromTemplate(UObject* TemplateObject);
+
+	/** Replace each SubSequence in this Assembly with new transient SubAssembly, using the existing SubSequence as a template */
+	void DuplicateSubSequences(ULevelSequence* SequenceTemplate = nullptr);
+
+	/**
+	 * Find the path of the input SubSequence relative to the outer sequence. If the outer sequence is a CineAssembly, then the relative path will be matched against
+	 * the folders known to belong to the Assembly's content tree, and a template (tokenized) path will be returned.
+	 */
+	FString GetRelativePath(ULevelSequence* SubSequence, ULevelSequence* OuterSequence);
+
 	/** Update the underlying json object whenever keys/values in the instance metadata map are added/removed/modified */
 	void UpdateInstanceMetadata();
 
@@ -191,9 +219,6 @@ private:
 	TOptional<FDateTime> TryGetCreatedAsLocalTime() const;
 
 #if WITH_EDITOR
-	/** Creates a new CineAssembly asset using the input object as a template */
-	UCineAssembly* CreateSubAssemblyFromTemplate(const FString& SubAssemblyName, UObject* TemplateObject);
-
 	/** Recursively gather the movie scene folders within a root folder */
 	void GetAllMovieSceneFoldersRecursive(UMovieSceneFolder* RootFolder, TArray<UMovieSceneFolder*>& OutFolders);
 #endif // WITH_EDITOR
@@ -203,6 +228,10 @@ public:
 	UPROPERTY()
 	FTemplateString AssemblyName;
 
+	/** The path where this assembly should be created, relative to its content root */
+	UPROPERTY()
+	FTemplateString RelativeAssemblyPath;
+
 	/** The level to open before opening this asset in Sequencer */
 	UPROPERTY(EditAnywhere, Category = "Default", meta = (AllowedClasses = "/Script/Engine.World"))
 	FSoftObjectPath Level;
diff --git a/Engine/Plugins/VirtualProduction/CinematicAssemblyTools/Source/CineAssemblyTools/Public/CineAssemblySchema.h b/Engine/Plugins/VirtualProduction/CinematicAssemblyTools/Source/CineAssemblyTools/Public/CineAssemblySchema.h
--- a/Engine/Plugins/VirtualProduction/CinematicAssemblyTools/Source/CineAssemblyTools/Public/CineAssemblySchema.h
+++ b/Engine/Plugins/VirtualProduction/CinematicAssemblyTools/Source/CineAssemblyTools/Public/CineAssemblySchema.h
@@ -116,10 +116,6 @@ public:
 	UPROPERTY(EditAnywhere, Category = "Metadata")
 	TArray<FAssemblyMetadataDesc> AssemblyMetadata;
 
-	/** Paths of subsequence assets that should be created for assemblies that use this schema, relative to the path of the top-level assembly */
-	UPROPERTY()
-	TArray<FString> SubsequencesToCreate;
-
 	/** Paths of folders that should be created for assemblies that use this schema, relative to the path of the top-level assembly */
 	UPROPERTY()
 	TArray<FString> FoldersToCreate;
@@ -128,14 +124,6 @@ public:
 	UPROPERTY()
 	TObjectPtr<UCineAssembly> TemplateSequence;
 
-	/** Path to another level sequence asset to use as a template when creating the assembly as a starting point. */
-	UPROPERTY()
-	FSoftObjectPath Template;
-
-	/** Templates to use for subassemblies that are created as part of this assembly. */
-	UPROPERTY()
-	TMap<FString, FSoftObjectPath> SubsequenceTemplates;
-
 	/** Mark this Schema as Data-Only. Assemblies made from this Schema will have their Data-Only defaulted to this value. */
 	UPROPERTY(AssetRegistrySearchable, BlueprintReadWrite, EditAnywhere, Category = "Default", AdvancedDisplay, meta = (DisplayName = "Data-Only"))
 	bool bIsDataOnly = false;
@@ -146,6 +134,21 @@ public:
 
 	static CINEASSEMBLYTOOLS_API const FName SchemaGuidPropertyName;
 
+protected:
+#if WITH_EDITORONLY_DATA 
+	UE_DEPRECATED(5.8, "Subsequences to create are now represented by the SubAssembly sections of the TemplateSequence.")
+	UPROPERTY()
+	TArray<FString> SubsequencesToCreate_DEPRECATED;
+
+	UE_DEPRECATED(5.8, "Use the TemplateSequence owned by the Schema to initialize a new Assembly.")
+	UPROPERTY()
+	FSoftObjectPath Template_DEPRECATED;
+
+	UE_DEPRECATED(5.8, "Subsequence Template objects are now stored in the SubAssembly sections of the TemplateSequence.")
+	UPROPERTY()
+	TMap<FString, FSoftObjectPath> SubsequenceTemplates_DEPRECATED;
+#endif // WITH_EDITORONLY_DATA 
+
 private:
 	/** Unique ID for this schema, assigned at object creation */
 	UPROPERTY(AssetRegistrySearchable, meta = (IgnoreForMemberInitializationTest))
diff --git a/Engine/Plugins/VirtualProduction/CinematicAssemblyTools/Source/CineAssemblyTools/Public/MovieSceneSubAssemblySection.h b/Engine/Plugins/VirtualProduction/CinematicAssemblyTools/Source/CineAssemblyTools/Public/MovieSceneSubAssemblySection.h
--- a/Engine/Plugins/VirtualProduction/CinematicAssemblyTools/Source/CineAssemblyTools/Public/MovieSceneSubAssemblySection.h
+++ b/Engine/Plugins/VirtualProduction/CinematicAssemblyTools/Source/CineAssemblyTools/Public/MovieSceneSubAssemblySection.h
@@ -49,10 +49,20 @@ public:
 	/** Returns the Assembly Template object */
 	UE_API UObject* GetAssemblyTemplate() const;
 
-#if WITH_EDITOR
 	/** Sets the Assembly Template and updates the SequenceName to match the new template */
 	UE_API void SetAssemblyTemplate(UObject* TemplateObject);
-#endif //WITH_EDITOR
+
+	/** Get the name to use when creating a new SubAssembly for this section */
+	UE_API FText GetSequenceName() const;
+
+	/** Set the name to use when creating a new SubAssembly for this section */
+	UE_API void SetSequenceName(const FText& InName);
+
+	/** Get the relative path to use when creating a new SubAssembly for this section */
+	UE_API FString GetSequencePath() const;
+
+	/** Set the relative path to use when creating a new SubAssembly for this section */
+	UE_API void SetSequencePath(const FString& InPath);
 
 public:
 	/** Indicates whether this section will be used as a reference to an existing sequence or as a template to create a new sequence */
@@ -63,6 +73,10 @@ public:
 	UPROPERTY()
 	FText NewSequenceName;
 
+	/** The path (relative to the root assembly's path) where the new sequence asset for this sections should be created */
+	UPROPERTY()
+	FString NewSequencePath;
+
 	/** Name of the protected AssemblyTemplate property */
 	static UE_API const FName AssemblyTemplatePropertyName;
 
diff --git a/Engine/Plugins/VirtualProduction/CinematicAssemblyTools/Source/CineAssemblyToolsEditor/Private/CineAssemblyCustomization.cpp b/Engine/Plugins/VirtualProduction/CinematicAssemblyTools/Source/CineAssemblyToolsEditor/Private/CineAssemblyCustomization.cpp
--- a/Engine/Plugins/VirtualProduction/CinematicAssemblyTools/Source/CineAssemblyToolsEditor/Private/CineAssemblyCustomization.cpp
+++ b/Engine/Plugins/VirtualProduction/CinematicAssemblyTools/Source/CineAssemblyToolsEditor/Private/CineAssemblyCustomization.cpp
@@ -322,12 +322,7 @@ void FCineAssemblyCustomization::CustomizeSubsequenceCategory(IDetailLayoutBuild
 	IDetailCategoryBuilder& SubsequenceCategory = DetailBuilder.EditCategory("SchemaSubsequences", SubsequencesCategoryName);
 
 	SubAssemblyNames.Reset();
-	for (const FString& SubsequenceName : BaseSchema->SubsequencesToCreate)
-	{
-		FTemplateString SubsequenceTemplateString;
-		SubsequenceTemplateString.Template = SubsequenceName;
-		SubAssemblyNames.Add(SubsequenceTemplateString);
-	}
+	// TODO: Fix this customization to display the SubAssembly names from the Customized Assembly's SubSequences
 
 	for (int32 Index = 0; Index < SubAssemblyNames.Num(); ++Index)
 	{
diff --git a/Engine/Plugins/VirtualProduction/CinematicAssemblyTools/Source/CineAssemblyToolsEditor/Private/CineAssemblyFactory.cpp b/Engine/Plugins/VirtualProduction/CinematicAssemblyTools/Source/CineAssemblyToolsEditor/Private/CineAssemblyFactory.cpp
--- a/Engine/Plugins/VirtualProduction/CinematicAssemblyTools/Source/CineAssemblyToolsEditor/Private/CineAssemblyFactory.cpp
+++ b/Engine/Plugins/VirtualProduction/CinematicAssemblyTools/Source/CineAssemblyToolsEditor/Private/CineAssemblyFactory.cpp
@@ -14,6 +14,8 @@
 #include "Modules/ModuleManager.h"
 #include "MovieScene.h"
 #include "MovieSceneToolsProjectSettings.h"
+#include "Sections/MovieSceneSubSection.h"
+#include "Tracks/MovieSceneSubTrack.h"
 #include "UI/CineAssembly/SCineAssemblyConfigWindow.h"
 
 #define LOCTEXT_NAMESPACE "CineAssemblyFactory"
@@ -143,34 +145,6 @@ void UCineAssemblyFactory::CreateConfiguredAssembly(UCineAssembly* ConfiguredAss
 	ConfiguredAssembly->SetFlags(Flags);
 	ConfiguredAssembly->ClearFlags(RF_Transient);
 
-	// If the configured assembly already has a MovieScene object (which could happen in case of a duplication), do not overwrite it with a new one
-	if (!ConfiguredAssembly->GetMovieScene())
-	{
-		// if we have a template to build from, setup now
-		bool bBuiltFromTemplate = false;
-		if (const UCineAssemblySchema* Schema = ConfiguredAssembly->GetSchema())
-		{
-			ULevelSequence* Template = Cast<ULevelSequence>(Schema->Template.TryLoad());
-			if (!Template)
-			{
-				Template = Schema->TemplateSequence;
-			}
-
-			if (Template)
-			{
-				ConfiguredAssembly->InitializeFromTemplate(Template);
-				// Confirm we have a valid MovieScene built.
-				bBuiltFromTemplate = ConfiguredAssembly->GetMovieScene() != nullptr;
-			}
-		}
-
-		if (!bBuiltFromTemplate)
-		{
-			// Do the same setup that ULevelSequence assets do when they are created
-			ConfiguredAssembly->Initialize();
-		}
-	}
-
 	// Re-evaluate the assembly's metadata tokens
 	if (const UCineAssemblySchema* Schema = ConfiguredAssembly->GetSchema())
 	{
@@ -188,7 +162,11 @@ void UCineAssemblyFactory::CreateConfiguredAssembly(UCineAssembly* ConfiguredAss
 		}
 	}
 
-	ConfiguredAssembly->CreateSubAssemblies();
+	ConfiguredAssembly->CreateSubFolders();
+	ConfiguredAssembly->ResolveMovieSceneTokens();
+
+	// Create packages for each of the transient SubAssembly
+	CreateConfiguredSubAssemblies(ConfiguredAssembly, CreateAssetPath);
 
 	// Analytics
 	UE::CineAssemblyToolsAnalytics::RecordEvent_CreateAssembly(ConfiguredAssembly);
@@ -204,4 +182,33 @@ void UCineAssemblyFactory::CreateConfiguredAssembly(UCineAssembly* ConfiguredAss
 	ContentBrowser.SetSelectedPaths({ CreateAssetPath }, true);
 }
 
+void UCineAssemblyFactory::CreateConfiguredSubAssemblies(UCineAssembly* ConfiguredAssembly, const FString& CreateAssetPath)
+{
+	for (UMovieSceneTrack* Track : ConfiguredAssembly->GetMovieScene()->GetTracks())
+	{
+		if (UMovieSceneSubTrack* SubTrack = Cast<UMovieSceneSubTrack>(Track))
+		{
+			for (UMovieSceneSection* Section : SubTrack->GetAllSections())
+			{
+				if (UMovieSceneSubSection* SubSection = Cast<UMovieSceneSubSection>(Section))
+				{
+					// If the SubSequence is a transient CineAssembly, create a new package for it
+					UCineAssembly* SubAssembly = Cast<UCineAssembly>(SubSection->GetSequence());
+					if (SubAssembly && SubAssembly->GetPackage() == GetTransientPackage())
+					{
+						SubAssembly->ParentAssembly = ConfiguredAssembly;
+						SubAssembly->Level = ConfiguredAssembly->Level;
+						SubAssembly->Production = ConfiguredAssembly->Production;
+						SubAssembly->ProductionName = ConfiguredAssembly->ProductionName;
+
+						SubAssembly->RelativeAssemblyPath.Resolved = UCineAssemblyNamingTokens::GetResolvedText(SubAssembly->RelativeAssemblyPath.Template, ConfiguredAssembly);
+
+						CreateConfiguredAssembly(SubAssembly, CreateAssetPath / SubAssembly->RelativeAssemblyPath.Resolved.ToString());
+					}
+				}
+			}
+		}
+	}
+}
+
 #undef LOCTEXT_NAMESPACE
diff --git a/Engine/Plugins/VirtualProduction/CinematicAssemblyTools/Source/CineAssemblyToolsEditor/Private/CineAssemblyFactory.h b/Engine/Plugins/VirtualProduction/CinematicAssemblyTools/Source/CineAssemblyToolsEditor/Private/CineAssemblyFactory.h
--- a/Engine/Plugins/VirtualProduction/CinematicAssemblyTools/Source/CineAssemblyToolsEditor/Private/CineAssemblyFactory.h
+++ b/Engine/Plugins/VirtualProduction/CinematicAssemblyTools/Source/CineAssemblyToolsEditor/Private/CineAssemblyFactory.h
@@ -20,6 +20,9 @@ public:
 	/** Takes a pre-configured, transient assembly, creates a valid package for it, and initializes it */
 	static void CreateConfiguredAssembly(UCineAssembly* ConfiguredAssembly, const FString& CreateAssetPath);
 
+	/** Create a valid package for each SubAssembly found in a SubSequence or Shot Track of the input Assembly's MovieScene */
+	static void CreateConfiguredSubAssemblies(UCineAssembly* ConfiguredAssembly, const FString& CreateAssetPath);
+
 	/** Takes a CineAssembly and a creation path and makes a formatted CineAssembly package name. */
 	static FString MakeAssemblyPackageName(UCineAssembly* ConfiguredAssembly, const FString& CreateAssetPath);
 
diff --git a/Engine/Plugins/VirtualProduction/CinematicAssemblyTools/Source/CineAssemblyToolsEditor/Private/ProductionFunctionLibrary.cpp b/Engine/Plugins/VirtualProduction/CinematicAssemblyTools/Source/CineAssemblyToolsEditor/Private/ProductionFunctionLibrary.cpp
--- a/Engine/Plugins/VirtualProduction/CinematicAssemblyTools/Source/CineAssemblyToolsEditor/Private/ProductionFunctionLibrary.cpp
+++ b/Engine/Plugins/VirtualProduction/CinematicAssemblyTools/Source/CineAssemblyToolsEditor/Private/ProductionFunctionLibrary.cpp
@@ -126,7 +126,8 @@ UCineAssembly* UProductionFunctionLibrary::CreateAssembly(UCineAssemblySchema* S
 	}
 
 	UCineAssembly* NewAssembly = NewObject<UCineAssembly>(GetTransientPackage(), NAME_None, RF_Transient);
-	NewAssembly->SetSchema(Schema);
+	NewAssembly->InitializeFromSchema(Schema);
+
 	NewAssembly->SetLevel(Level);
 	NewAssembly->SetParentAssembly(ParentAssembly);
 
diff --git a/Engine/Plugins/VirtualProduction/CinematicAssemblyTools/Source/CineAssemblyToolsEditor/Private/Sequencer/SubAssemblySection.cpp b/Engine/Plugins/VirtualProduction/CinematicAssemblyTools/Source/CineAssemblyToolsEditor/Private/Sequencer/SubAssemblySection.cpp
--- a/Engine/Plugins/VirtualProduction/CinematicAssemblyTools/Source/CineAssemblyToolsEditor/Private/Sequencer/SubAssemblySection.cpp
+++ b/Engine/Plugins/VirtualProduction/CinematicAssemblyTools/Source/CineAssemblyToolsEditor/Private/Sequencer/SubAssemblySection.cpp
@@ -112,7 +112,7 @@ FText FSubAssemblySection::GetSectionText() const
 {
 	if (MovieSceneSubAssemblySection->IsTemplateSection())
 	{
-		return MovieSceneSubAssemblySection->NewSequenceName;
+		return MovieSceneSubAssemblySection->GetSequenceName();
 	}
 
 	return MovieSceneSubAssemblySection->GetSequence() ? MovieSceneSubAssemblySection->GetSequence()->GetDisplayName() : FText::GetEmpty();
@@ -122,7 +122,7 @@ void FSubAssemblySection::SetSectionText(const FText& InText, ETextCommit::Type
 {
 	if (MovieSceneSubAssemblySection->IsTemplateSection())
 	{
-		MovieSceneSubAssemblySection->NewSequenceName = InText;
+		MovieSceneSubAssemblySection->SetSequenceName(InText);
 	}
 }
 
diff --git a/Engine/Plugins/VirtualProduction/CinematicAssemblyTools/Source/CineAssemblyToolsEditor/Private/Sequencer/SubAssemblyTrackEditor.cpp b/Engine/Plugins/VirtualProduction/CinematicAssemblyTools/Source/CineAssemblyToolsEditor/Private/Sequencer/SubAssemblyTrackEditor.cpp
--- a/Engine/Plugins/VirtualProduction/CinematicAssemblyTools/Source/CineAssemblyToolsEditor/Private/Sequencer/SubAssemblyTrackEditor.cpp
+++ b/Engine/Plugins/VirtualProduction/CinematicAssemblyTools/Source/CineAssemblyToolsEditor/Private/Sequencer/SubAssemblyTrackEditor.cpp
@@ -350,31 +350,27 @@ TSharedRef<SWidget> FSubAssemblyTrackEditor::BuildAssetPicker(UMovieSceneSubAsse
 
 void FSubAssemblyTrackEditor::AddNewSequenceToTrack(UMovieSceneSubAssemblyTrack* Track, UObject* TemplateObject)
 {
-	// Create a new empty Assembly and add it to the that can be added to the SubAssemblyTrack
-	UCineAssembly* NewAssembly = NewObject<UCineAssembly>();
-	NewAssembly->Initialize();
-
-	// Set the playback range of the new assembly to match the template object (if possible)
-	if (UMovieSceneSequence* Sequence = Cast<UMovieSceneSequence>(TemplateObject))
+	UMovieScene* FocusedMovieScene = GetFocusedMovieScene();
+	if (!FocusedMovieScene)
 	{
-		if (NewAssembly->GetMovieScene() && Sequence->GetMovieScene())
-		{
-			NewAssembly->GetMovieScene()->SetPlaybackRange(Sequence->GetMovieScene()->GetPlaybackRange());
-		}
+		return;
 	}
 
-	FSubTrackEditor::AddSequenceToTrack(NewAssembly, Track);
+	const FFrameNumber StartFrame = FocusedMovieScene->GetPlaybackRange().GetLowerBoundValue();
+	int32 Duration = FocusedMovieScene->GetPlaybackRange().Size<FFrameNumber>().Value;
 
-	// Find the SubAssemblySection in the Track that was just created (i.e. whose Sequence matches the one just created) and set its Assembly Template to be the input template object
-	for (UMovieSceneSection* Section : Track->GetAllSections())
+	// If the template object is a valid sequence, use its duration instead of the default duration
+	UMovieSceneSequence* TemplateSequence = Cast<UMovieSceneSequence>(TemplateObject);
+	if (TemplateSequence && TemplateSequence->GetMovieScene())
 	{
-		UMovieSceneSubAssemblySection* SubAssemblySection = Cast<UMovieSceneSubAssemblySection>(Section);
-		if (SubAssemblySection->GetSequence() == NewAssembly)
-		{
-			SubAssemblySection->SetAssemblyTemplate(TemplateObject);
-			break;
-		}
+		Duration = TemplateSequence->GetMovieScene()->GetPlaybackRange().Size<FFrameNumber>().Value;
 	}
+
+	// Add a new SubAssembly section to the input track, and set its assembly template object
+	UMovieSceneSubAssemblySection* SubAssemblySection = Cast<UMovieSceneSubAssemblySection>(Track->AddSequence(nullptr, StartFrame, Duration));
+	SubAssemblySection->SetAssemblyTemplate(TemplateObject);
+
+	GetSequencer()->NotifyMovieSceneDataChanged(EMovieSceneDataChangeType::MovieSceneStructureItemsChanged);
 }
 
 void FSubAssemblyTrackEditor::OnAssetSelectedAsTemplate(const FAssetData& AssetData, UMovieSceneSubAssemblyTrack* Track)
diff --git a/Engine/Plugins/VirtualProduction/CinematicAssemblyTools/Source/CineAssemblyToolsEditor/Private/TakeRecorder/CineAssemblyTakeRecorderIntegration.cpp b/Engine/Plugins/VirtualProduction/CinematicAssemblyTools/Source/CineAssemblyToolsEditor/Private/TakeRecorder/CineAssemblyTakeRecorderIntegration.cpp
--- a/Engine/Plugins/VirtualProduction/CinematicAssemblyTools/Source/CineAssemblyToolsEditor/Private/TakeRecorder/CineAssemblyTakeRecorderIntegration.cpp
+++ b/Engine/Plugins/VirtualProduction/CinematicAssemblyTools/Source/CineAssemblyToolsEditor/Private/TakeRecorder/CineAssemblyTakeRecorderIntegration.cpp
@@ -66,18 +66,24 @@ void FCineAssemblyTakeRecorderIntegration::OnRecordingInitialized(UTakeRecorder*
 			}
 
 			UCineAssemblyTakeRecorderSettings* Settings = GetMutableDefault<UCineAssemblyTakeRecorderSettings>();
+			UCineAssemblySchema* Schema = nullptr;
+
 			if (!Settings->AssemblySchema.IsNull())
 			{
-				UCineAssemblySchema* Schema = Settings->AssemblySchema.LoadSynchronous();
-				Assembly->SetSchema(Schema);
+				Schema = Settings->AssemblySchema.LoadSynchronous();
+			}
+
+			if (Schema)
+			{
+				Assembly->InitializeFromSchema(Schema);
+
+				FString CurrentAssemblyPath;
+				Assembly->GetPathName(nullptr, CurrentAssemblyPath);
+				CurrentAssemblyPath = FPaths::GetPath(CurrentAssemblyPath);
 
 				// If the schema defines some additional default assembly path, move the assembly asset accordingly
 				if (!Schema->DefaultAssemblyPath.IsEmpty())
 				{
-					FString CurrentAssemblyPath;
-					Assembly->GetPathName(nullptr, CurrentAssemblyPath);
-					CurrentAssemblyPath = FPaths::GetPath(CurrentAssemblyPath);
-
 					FString UniquePackageName;
 					FString UniqueAssetName;
 					UCineAssemblyFactory::MakeUniqueNameAndPath(Assembly, CurrentAssemblyPath, UniquePackageName, UniqueAssetName);
@@ -86,7 +92,9 @@ void FCineAssemblyTakeRecorderIntegration::OnRecordingInitialized(UTakeRecorder*
 					Assembly->Rename(*UniqueAssetName, Package);
 				}
 
-				Assembly->CreateSubAssemblies();
+				Assembly->CreateSubFolders();
+				Assembly->ResolveMovieSceneTokens();
+				UCineAssemblyFactory::CreateConfiguredSubAssemblies(Assembly, CurrentAssemblyPath);
 			}
 
 			UE::CineAssemblyToolsAnalytics::RecordEvent_RecordAssembly();
diff --git a/Engine/Plugins/VirtualProduction/CinematicAssemblyTools/Source/CineAssemblyToolsEditor/Private/UI/CineAssembly/SCineAssemblyConfigWindow.cpp b/Engine/Plugins/VirtualProduction/CinematicAssemblyTools/Source/CineAssemblyToolsEditor/Private/UI/CineAssembly/SCineAssemblyConfigWindow.cpp
--- a/Engine/Plugins/VirtualProduction/CinematicAssemblyTools/Source/CineAssemblyToolsEditor/Private/UI/CineAssembly/SCineAssemblyConfigWindow.cpp
+++ b/Engine/Plugins/VirtualProduction/CinematicAssemblyTools/Source/CineAssemblyToolsEditor/Private/UI/CineAssembly/SCineAssemblyConfigWindow.cpp
@@ -29,6 +29,8 @@ void SCineAssemblyConfigWindow::Construct(const FArguments& InArgs, const FStrin
 
 	// Create a new transient CineAssembly to configure in the UI.
 	CineAssemblyToConfigure = TStrongObjectPtr<UCineAssembly>(NewObject<UCineAssembly>(GetTransientPackage(), NAME_None, RF_Transient));
+	CineAssemblyToConfigure->Initialize();
+
 	if (UMovieSceneMetaData* MetaDataToConfigure = CineAssemblyToConfigure->FindOrAddMetaData<UMovieSceneMetaData>())
 	{
 		// Default the MetaData for the CineAssembly being configured.
@@ -202,11 +204,12 @@ void SCineAssemblyConfigWindow::OnSchemaSelected(const FAssetData& AssetData)
 	if (AssetData.AssetName == NoSchemaName)
 	{
 		CineAssemblyToConfigure->ChangeSchema(nullptr);
+		CineAssemblyToConfigure->Initialize();
 		ConfigPanel->Refresh();
 	}
 	else if (UCineAssemblySchema* CineAssemblySchema = Cast<UCineAssemblySchema>(AssetData.GetAsset()))
 	{
-		CineAssemblyToConfigure->ChangeSchema(CineAssemblySchema);
+		CineAssemblyToConfigure->InitializeFromSchema(CineAssemblySchema);
 		ConfigPanel->Refresh();
 	}
 }
diff --git a/Engine/Plugins/VirtualProduction/CinematicAssemblyTools/Source/CineAssemblyToolsEditor/Private/UI/CineAssembly/SCineAssemblySchemaWindow.cpp b/Engine/Plugins/VirtualProduction/CinematicAssemblyTools/Source/CineAssemblyToolsEditor/Private/UI/CineAssembly/SCineAssemblySchemaWindow.cpp
--- a/Engine/Plugins/VirtualProduction/CinematicAssemblyTools/Source/CineAssemblyToolsEditor/Private/UI/CineAssembly/SCineAssemblySchemaWindow.cpp
+++ b/Engine/Plugins/VirtualProduction/CinematicAssemblyTools/Source/CineAssemblyToolsEditor/Private/UI/CineAssembly/SCineAssemblySchemaWindow.cpp
@@ -16,6 +16,7 @@
 #include "Misc/MessageDialog.h"
 #include "Modules/ModuleManager.h"
 #include "MovieScene.h"
+#include "MovieSceneSubAssemblySection.h"
 #include "MovieSceneSubAssemblyTrack.h"
 #include "NamingTokensEngineSubsystem.h"
 #include "ProductionSettings.h"
@@ -50,9 +51,24 @@ struct ISchemaContentHierarchyTreeItem : TSharedFromThis<ISchemaContentHierarchy
 
 	/** Whether this item can accept an item being dropped onto it */
 	virtual TOptional<EItemDropZone> CanAcceptDrop() const { return TOptional<EItemDropZone>(); }
+	
+	/** Return the name of the tree item */
+	virtual FString GetName() const = 0;
 
-	/** The fully qualified path of this item (relative to the root folder) */
-	FString Path;
+	/** Set the name of the tree item */
+	virtual void SetName(const FString& InName) = 0;
+
+	/** Return the display name of the tree item */
+	virtual FText GetDisplayName() const = 0;
+
+	/** Set the display name of the tree item */
+	virtual void SetDisplayName(const FText& InDisplayName) = 0;
+
+	/** Return the path of the tree item (relative to the root folder) */
+	virtual FString GetPath() const = 0;
+
+	/** Set the path of the tree item (relative to the root folder) */
+	virtual void SetPath(const FString& InPath) = 0;
 
 	/** The parent folder of this item */
 	TSharedPtr<FSchemaContentHierarchyFolder> ParentFolder;
@@ -61,13 +77,105 @@ struct ISchemaContentHierarchyTreeItem : TSharedFromThis<ISchemaContentHierarchy
 	TSharedPtr<SNamingTokensEditableTextBox> NameWidget;
 };
 
+struct FSchemaContentHierarchyAsset : ISchemaContentHierarchyTreeItem
+{
+	//~ Begin ISchemaContentHierarchyTreeItem Interface
+	virtual TSharedPtr<FSchemaContentHierarchyAsset> AsAsset() override { return SharedThis(this); }
+
+	virtual FString GetName() const override
+	{
+		return SubAssemblySection.Get() ? SubAssemblySection.Get()->GetSequenceName().ToString() : TEXT("");
+	}
+
+	virtual void SetName(const FString& InName) override
+	{
+		if (SubAssemblySection.IsValid())
+		{
+			SubAssemblySection->SetSequenceName(FText::FromString(InName));
+		}
+	}
+
+	virtual FText GetDisplayName() const override
+	{
+		return SubAssemblySection.Get() ? SubAssemblySection.Get()->GetSequenceName() : FText::GetEmpty();
+	}
+
+	virtual void SetDisplayName(const FText& InDisplayName) override
+	{
+		if (SubAssemblySection.IsValid())
+		{
+			SubAssemblySection->SetSequenceName(InDisplayName);
+		}
+	}
+
+	virtual FString GetPath() const override
+	{
+		return SubAssemblySection.Get() ? SubAssemblySection.Get()->GetSequencePath() : TEXT("");
+	}
+
+	virtual void SetPath(const FString& InPath) override
+	{
+		if (SubAssemblySection.IsValid())
+		{
+			SubAssemblySection->SetSequencePath(FPaths::GetPath(InPath));
+		}
+	}
+	//~ End ISchemaContentHierarchyTreeItem Interface
+
+	/** Returns the path name of the item's Template Object */
+	FSoftObjectPath GetTemplateObjectPath() const
+	{
+		if (SubAssemblySection.IsValid() && SubAssemblySection->GetAssemblyTemplate())
+		{
+			return SubAssemblySection->GetAssemblyTemplate()->GetPathName();
+		}
+		return FSoftObjectPath();
+	}
+
+	/** The MovieScene section represented in the content tree by this asset item */
+	TWeakObjectPtr<UMovieSceneSubAssemblySection> SubAssemblySection;
+};
+
 struct FSchemaContentHierarchyFolder : ISchemaContentHierarchyTreeItem
 {
 	//~ Begin ISchemaContentHierarchyTreeItem Interface
 	virtual TSharedPtr<FSchemaContentHierarchyFolder> AsFolder() override { return SharedThis(this); }
 	virtual TOptional<EItemDropZone> CanAcceptDrop() const override { return EItemDropZone::OntoItem; }
+
+	virtual FString GetName() const override
+	{
+		return FPaths::GetPathLeaf(FullPath);
+	}
+
+	virtual void SetName(const FString& InName) override
+	{
+		FullPath = FPaths::GetPath(FullPath) / InName;
+	}
+
+	virtual FText GetDisplayName() const override
+	{
+		return FText::FromString(GetName());
+	}
+
+	virtual void SetDisplayName(const FText& InDisplayName) override
+	{
+		FullPath = FPaths::GetPath(FullPath) / InDisplayName.ToString();
+	}
+
+	virtual FString GetPath() const override
+	{
+		return FullPath;
+	}
+
+	virtual void SetPath(const FString& InPath) override
+	{
+		FullPath = InPath;
+	}
 	//~ End ISchemaContentHierarchyTreeItem Interface
 
+	/** The fully qualified path of this item (relative to the root folder) */
+	FString FullPath;
+
 	/** The child assets of this folder item */
 	TArray<TSharedPtr<FSchemaContentHierarchyAsset>> ChildAssets;
 
@@ -75,19 +183,9 @@ struct FSchemaContentHierarchyFolder : ISchemaContentHierarchyTreeItem
 	TArray<TSharedPtr<FSchemaContentHierarchyFolder>> ChildFolders;
 };
 
-struct FSchemaContentHierarchyAsset : ISchemaContentHierarchyTreeItem
-{
-	//~ Begin ISchemaContentHierarchyTreeItem Interface
-	virtual TSharedPtr<FSchemaContentHierarchyAsset> AsAsset() override { return SharedThis(this); }
-	//~ End ISchemaContentHierarchyTreeItem Interface
-
-	/** Asset Template soft path for this asset item */
-	FSoftObjectPath AssetTemplate;
-};
-
 namespace UE::CineAssemblyTools::Private
 {
-	auto SortTreeItems = [](const TSharedPtr<ISchemaContentHierarchyTreeItem>& A, const TSharedPtr<ISchemaContentHierarchyTreeItem>& B) { return A->Path < B->Path; };
+	auto SortTreeItems = [](const TSharedPtr<ISchemaContentHierarchyTreeItem>& A, const TSharedPtr<ISchemaContentHierarchyTreeItem>& B) { return A->GetPath() < B->GetPath(); };
 };
 
 void SCineAssemblySchemaWindow::Construct(const FArguments& InArgs, const FString& InCreateAssetPath)
@@ -749,16 +847,15 @@ TSharedRef<SWidget> SCineAssemblySchemaWindow::MakeHierarchyTabContent()
 
 FReply SCineAssemblySchemaWindow::OnAddAssetClicked()
 {
-	// Add a new empty string to the list, which will be renamed by the user
-	AssetListItems.Add(MakeShared<FSchemaAssetParameters>(FSchemaAssetParameters()));
-	AssetListView->RequestListRefresh();
+	// TODO: Add support for adding any asset type to the schema definition (not just subsequences)
+	// TODO: When adding subsequence assets, they will also be added as new tracks to the TemplateSequence MovieScene
 	return FReply::Handled();
 }
 
 FString SCineAssemblySchemaWindow::MakeUniqueFolderPath(TSharedPtr<FSchemaContentHierarchyFolder> FolderItem)
 {
 	// This implementation is based on a similar utility in AssetTools for creating a unique asset name.
-	const FString BaseName = FolderItem->Path / TEXT("NewFolder");
+	const FString BaseName = FolderItem->GetPath() / TEXT("NewFolder");
 
 	// Find the index in the string of the last non-numeric character
 	int32 CharIndex = BaseName.Len() - 1;
@@ -780,7 +877,7 @@ FString SCineAssemblySchemaWindow::MakeUniqueFolderPath(TSharedPtr<FSchemaConten
 
 	FString WorkingName = TrimmedBaseName;
 
-	while (Algo::ContainsBy(FolderItem->ChildFolders, WorkingName, &FSchemaContentHierarchyFolder::Path))
+	while (Algo::ContainsBy(FolderItem->ChildFolders, WorkingName, &FSchemaContentHierarchyFolder::FullPath))
 	{
 		WorkingName = FString::Printf(TEXT("%s%d"), *TrimmedBaseName, IntSuffix);
 		IntSuffix++;
@@ -803,7 +900,7 @@ FReply SCineAssemblySchemaWindow::OnAddFolderClicked()
 		ParentFolder = SelectedItems[0]->AsFolder() ? SelectedItems[0]->AsFolder() : SelectedItems[0]->ParentFolder;
 	}
 
-	NewFolder->Path = MakeUniqueFolderPath(ParentFolder);
+	NewFolder->SetPath(MakeUniqueFolderPath(ParentFolder));
 	NewFolder->ParentFolder = ParentFolder;
 
 	ParentFolder->ChildFolders.Add(NewFolder);
@@ -1015,10 +1112,7 @@ TSharedRef<ITableRow> SCineAssemblySchemaWindow::OnGenerateTreeRow(TSharedPtr<IS
 						.DisplayTokenIcon(false)
 						.DisplayBorderImage(false)
 						.DisplayErrorMessage(!bIsRootItem)
-						.Text_Lambda([bIsRootItem, TreeItem]() -> FText
-							{
-								return bIsRootItem ? LOCTEXT("RootPathName", "Root Folder") : FText::FromString(FPaths::GetPathLeaf(TreeItem->Path));
-							})
+						.Text(this, &SCineAssemblySchemaWindow::GetTreeItemDisplayName, TreeItem)
 						.OnValidateTokenizedText_Lambda([this, TreeItem](const FText& InText, FText& OutErrorMessage) -> bool
 							{
 								return TreeItem->AsFolder() ? ValidateFolderName(InText, OutErrorMessage, TreeItem->AsFolder()) : ValidateAssetName(InText, OutErrorMessage);
@@ -1036,20 +1130,39 @@ TSharedRef<ITableRow> SCineAssemblySchemaWindow::OnGenerateTreeRow(TSharedPtr<IS
 						.DisplayCompactSize(true)
 						.AllowedClass(ULevelSequence::StaticClass())
 						.ToolTipText(LOCTEXT("SchemaAssetTemplateSelectionTooltip", "(Optional) Select an existing level sequence to use as a template for this subsequence."))
-						.Visibility(TreeItem->AsAsset() ? EVisibility::Visible : EVisibility::Collapsed)
-						.ObjectPath_Lambda([WeakTreeItem = TWeakPtr<ISchemaContentHierarchyTreeItem>(TreeItem)]() -> FString
-							{
-								if (TSharedPtr<ISchemaContentHierarchyTreeItem> ItemPtr = WeakTreeItem.Pin())
-								{
-									return ItemPtr->AsAsset() ? ItemPtr->AsAsset()->AssetTemplate.ToString() : FString();
-								}
-								return FString();
-							})
-						.OnObjectChanged(this, &SCineAssemblySchemaWindow::OnTreeItemTemplateChanged, TreeItem)
+						.Visibility(TreeItem->AsAsset() && TreeItem != TopLevelAssemblyItem ? EVisibility::Visible : EVisibility::Collapsed)
+						.ObjectPath(this, &SCineAssemblySchemaWindow::GetTemplateObjectPath, TreeItem->AsAsset())
+						.OnObjectChanged(this, &SCineAssemblySchemaWindow::OnTreeItemTemplateChanged, TreeItem->AsAsset())
 				]
 		];
 }
 
+FText SCineAssemblySchemaWindow::GetTreeItemDisplayName(TSharedPtr<ISchemaContentHierarchyTreeItem> TreeItem) const
+{
+	if (TreeItem == RootFolder)
+	{
+		return LOCTEXT("RootPathName", "Root Folder");
+	}
+	else if (TreeItem == TopLevelAssemblyItem)
+	{
+		return FText::FromString(TEXT("{assembly}"));
+	}
+	else
+	{
+		return TreeItem->GetDisplayName();
+	}
+}
+
+FString SCineAssemblySchemaWindow::GetTemplateObjectPath(TSharedPtr<FSchemaContentHierarchyAsset> AssetItem) const
+{
+	if (!AssetItem || AssetItem == TopLevelAssemblyItem)
+	{
+		return FString();
+	}
+
+	return AssetItem->GetTemplateObjectPath().ToString();
+}
+
 void SCineAssemblySchemaWindow::OnGetChildren(TSharedPtr<ISchemaContentHierarchyTreeItem> TreeItem, TArray<TSharedPtr<ISchemaContentHierarchyTreeItem>>& OutNodes)
 {
 	// Display all of the child assets first, followed by all of the child folders
@@ -1155,7 +1268,7 @@ void SCineAssemblySchemaWindow::InitializeContentTree()
 {
 	// Create the content tree root
 	RootFolder = MakeShared<FSchemaContentHierarchyFolder>();
-	RootFolder->Path = TEXT("");
+	RootFolder->SetPath(TEXT(""));
 
 	TreeItems.Add(RootFolder);
 
@@ -1171,44 +1284,20 @@ void SCineAssemblySchemaWindow::InitializeContentTree()
 		if (ParentFolder)
 		{
 			TSharedPtr<FSchemaContentHierarchyFolder> NewItem = MakeShared<FSchemaContentHierarchyFolder>();
-			NewItem->Path = FolderPath;
+			NewItem->SetPath(FolderPath);
 			NewItem->ParentFolder = ParentFolder;
 
 			ParentFolder->ChildFolders.Add(NewItem);
 		}
 	}
 
-	// Sort the list of subsequences so that paths are added to the tree in the proper order
-	SchemaToConfigure->SubsequencesToCreate.Sort();
-
-	for (const FString& SubsequencePath : SchemaToConfigure->SubsequencesToCreate)
-	{
-		const FString ParentPath = FPaths::GetPath(SubsequencePath);
-
-		// Walk the tree until we find an item whose path matches the parent path. The new tree item will be created as one of its children
-		TSharedPtr<FSchemaContentHierarchyFolder> ParentFolder = FindFolderByPathRecursive(RootFolder, ParentPath);
-		if (ParentFolder)
-		{
-			TSharedPtr<FSchemaContentHierarchyAsset> NewItem = MakeShared<FSchemaContentHierarchyAsset>();
-			NewItem->Path = SubsequencePath;
-			NewItem->ParentFolder = ParentFolder;
-
-			if (const FSoftObjectPath* FoundTemplate = SchemaToConfigure->SubsequenceTemplates.Find(SubsequencePath))
-			{
-				NewItem->AssetTemplate = *FoundTemplate;
-			}
-
-			ParentFolder->ChildAssets.Add(NewItem);
-		}
-	}
+	UpdateContentHierarchyAssets();
 
 	// Add the top-level assembly node
 	TSharedPtr<FSchemaContentHierarchyFolder> ParentFolder = FindFolderByPathRecursive(RootFolder, SchemaToConfigure->DefaultAssemblyPath);
 	if (ParentFolder)
 	{
 		TopLevelAssemblyItem = MakeShared<FSchemaContentHierarchyAsset>();
-		TopLevelAssemblyItem->Path = TEXT("{assembly}");
-		TopLevelAssemblyItem->AssetTemplate = SchemaToConfigure->Template;
 		TopLevelAssemblyItem->ParentFolder = ParentFolder;
 
 		ParentFolder->ChildAssets.Add(TopLevelAssemblyItem);
@@ -1229,31 +1318,24 @@ void SCineAssemblySchemaWindow::GetFolderListRecursive(TSharedPtr<FSchemaContent
 {
 	for (const TSharedPtr<FSchemaContentHierarchyFolder>& ChildFolder : FolderItem->ChildFolders)
 	{
-		OutFolderList.Add(ChildFolder->Path);
+		OutFolderList.Add(ChildFolder->GetPath());
 		GetFolderListRecursive(ChildFolder, OutFolderList);
 	}
 }
 
-void SCineAssemblySchemaWindow::GetAssetListRecursive(TSharedPtr<FSchemaContentHierarchyFolder> FolderItem, TArray<FString>& OutAssetList, TMap<FString, FSoftObjectPath>& OutTemplates) const
+void SCineAssemblySchemaWindow::GetAssetListRecursive(TSharedPtr<FSchemaContentHierarchyFolder> FolderItem, TArray<TSharedPtr<FSchemaContentHierarchyAsset>>& OutAssetList) const
 {
-	for (const TSharedPtr<FSchemaContentHierarchyAsset>& Asset : FolderItem->ChildAssets)
-	{
-		if (Asset != TopLevelAssemblyItem)
-		{
-			OutAssetList.Add(Asset->Path);
-			OutTemplates.Add(Asset->Path, Asset->AssetTemplate);
-		}
-	}
+	OutAssetList.Append(FolderItem->ChildAssets);
 
 	for (const TSharedPtr<FSchemaContentHierarchyFolder>& ChildFolder : FolderItem->ChildFolders)
 	{
-		GetAssetListRecursive(ChildFolder, OutAssetList, OutTemplates);
+		GetAssetListRecursive(ChildFolder, OutAssetList);
 	}
 }
 
 TSharedPtr<FSchemaContentHierarchyFolder> SCineAssemblySchemaWindow::FindFolderByPathRecursive(TSharedPtr<FSchemaContentHierarchyFolder> FolderItem, const FString& FolderPath) const
 {
-	if (FolderItem->Path.Equals(FolderPath))
+	if (FolderItem->GetPath().Equals(FolderPath))
 	{
 		return FolderItem;
 	}
@@ -1271,20 +1353,15 @@ TSharedPtr<FSchemaContentHierarchyFolder> SCineAssemblySchemaWindow::FindFolderB
 
 void SCineAssemblySchemaWindow::SetChildrenPathRecursive(TSharedPtr<FSchemaContentHierarchyFolder> FolderItem, const FString& NewPath)
 {
-	for (TSharedPtr<FSchemaContentHierarchyAsset>& Asset : FolderItem->ChildAssets)
+	for (const TSharedPtr<FSchemaContentHierarchyAsset>& AssetItem : FolderItem->ChildAssets)
 	{
-		const FString OldAssetName = FPaths::GetPathLeaf(Asset->Path);
-
-		const FString NewAssetPath = NewPath / OldAssetName;
-		Asset->Path = NewAssetPath;
+		AssetItem->SetPath(NewPath / AssetItem->GetName());
 	}
 
-	for (TSharedPtr<FSchemaContentHierarchyFolder>& ChildFolder : FolderItem->ChildFolders)
+	for (const TSharedPtr<FSchemaContentHierarchyFolder>& ChildFolder : FolderItem->ChildFolders)
 	{
-		const FString OldChildFolderName = FPaths::GetPathLeaf(ChildFolder->Path);
-
-		const FString NewChildPath = NewPath / OldChildFolderName;
-		ChildFolder->Path = NewChildPath;
+		const FString NewChildPath = NewPath / ChildFolder->GetName();
+		ChildFolder->SetPath(NewChildPath);
 
 		SetChildrenPathRecursive(ChildFolder, NewChildPath);
 	}
@@ -1305,53 +1382,42 @@ void SCineAssemblySchemaWindow::EnterEditMode(TSharedPtr<FSchemaAssetParameters>
 void SCineAssemblySchemaWindow::OnTreeItemTextCommitted(const FText& InText, ETextCommit::Type InCommitType, TSharedPtr<ISchemaContentHierarchyTreeItem> TreeItem)
 {
 	// Early-out if the name has not actually changed
-	const FString OldPath = TreeItem->Path;
-	const FString OldName = FPaths::GetPathLeaf(OldPath);
+	const FString OldName = TreeItem->GetName();
 	if (OldName.Equals(InText.ToString()))
 	{
 		return;
 	}
 
-	const FString NewPath = FPaths::GetPath(OldPath) / InText.ToString();
-	TreeItem->Path = NewPath;
+	TreeItem->SetDisplayName(InText);
 
 	if (TSharedPtr<FSchemaContentHierarchyFolder> FolderItem = TreeItem->AsFolder())
 	{
 		// If this is a folder item, update the path of all of its children (recursively)
-		SetChildrenPathRecursive(FolderItem, NewPath);
+		SetChildrenPathRecursive(FolderItem, FolderItem->GetPath());
 
 		FolderItem->ParentFolder->ChildFolders.Sort(UE::CineAssemblyTools::Private::SortTreeItems);
+		UpdateFolderList();
 	}
 	else
 	{
 		TreeItem->ParentFolder->ChildAssets.Sort(UE::CineAssemblyTools::Private::SortTreeItems);
 	}
 
-	UpdateFolderList();
-	UpdateAssetList();
-
 	TreeView->RequestTreeRefresh();
 }
 
-void SCineAssemblySchemaWindow::OnTreeItemTemplateChanged(const FAssetData& InAssetData, TSharedPtr<ISchemaContentHierarchyTreeItem> TreeItem)
+void SCineAssemblySchemaWindow::OnTreeItemTemplateChanged(const FAssetData& InAssetData, TSharedPtr<FSchemaContentHierarchyAsset> AssetItem)
 {
-	// Early-out if this is not an asset item
-	TSharedPtr<FSchemaContentHierarchyAsset> AssetItem = TreeItem->AsAsset();
-	if (!AssetItem)
+	if (!AssetItem || AssetItem == TopLevelAssemblyItem)
 	{
 		return;
 	}
 
-	const FSoftObjectPath OldSoftPath = AssetItem->AssetTemplate;
-	const FSoftObjectPath NewSoftPath = InAssetData.GetSoftObjectPath();
-	// Early-out if the name has not actually changed
-	if (OldSoftPath == NewSoftPath)
+	if (AssetItem->SubAssemblySection.IsValid())
 	{
-		return;
+		AssetItem->SubAssemblySection->SetAssemblyTemplate(InAssetData.GetAsset());
 	}
 
-	AssetItem->AssetTemplate = NewSoftPath;
-	UpdateAssetList();
 	TreeView->RequestTreeRefresh();
 }
 
@@ -1397,7 +1463,7 @@ bool SCineAssemblySchemaWindow::ValidateFolderName(const FText& InText, FText& O
 	}
 
 	// Check for duplicate folder names
-	const FString ParentPath = FPaths::GetPath(FolderItem->Path);
+	const FString ParentPath = FPaths::GetPath(FolderItem->GetPath());
 	TSharedPtr<FSchemaContentHierarchyFolder> MatchingFolder = FindFolderByPathRecursive(RootFolder, ParentPath / InText.ToString());
 
 	if (MatchingFolder && MatchingFolder != FolderItem)
@@ -1423,23 +1489,75 @@ void SCineAssemblySchemaWindow::DeleteTreeItem(TSharedPtr<ISchemaContentHierarch
 {
 	if (TSharedPtr<FSchemaContentHierarchyFolder> FolderItem = TreeItem->AsFolder())
 	{
-		if (ContainsRecursive(FolderItem, TopLevelAssemblyItem))
-		{
-			FolderItem->ParentFolder->ChildAssets.Add(TopLevelAssemblyItem);
-			TopLevelAssemblyItem->ParentFolder = FolderItem->ParentFolder;
-		}
+		DeleteFolderItem(FolderItem);
+	}
+	else if (TSharedPtr<FSchemaContentHierarchyAsset> AssetItem = TreeItem->AsAsset())
+	{
+		DeleteAssetItem(AssetItem);
+	}
 
-		FolderItem->ParentFolder->ChildFolders.Remove(FolderItem);
+	if (TSharedPtr<FSchemaContentHierarchyFolder> FolderItem = TreeItem->AsFolder())
+	{
+		TreeItem->ParentFolder->ChildFolders.Remove(FolderItem);
 	}
-	else
+	else if (TSharedPtr<FSchemaContentHierarchyAsset> AssetItem = TreeItem->AsAsset())
 	{
-		TreeItem->ParentFolder->ChildAssets.Remove(TreeItem->AsAsset());
+		TreeItem->ParentFolder->ChildAssets.Remove(AssetItem);
 	}
 
 	UpdateFolderList();
-	UpdateAssetList();
 
 	TreeView->RequestTreeRefresh();
+
+	Sequencer->NotifyMovieSceneDataChanged(EMovieSceneDataChangeType::MovieSceneStructureItemRemoved);
+}
+
+void SCineAssemblySchemaWindow::DeleteFolderItem(TSharedPtr<FSchemaContentHierarchyFolder> FolderItem)
+{
+	// If the folder being deleted contains the TopLevelAssembly (which cannot be deleted), reparent the TopLevelAssembly first
+	if (ContainsRecursive(FolderItem, TopLevelAssemblyItem))
+	{
+		TopLevelAssemblyItem->ParentFolder = FolderItem->ParentFolder;
+		FolderItem->ParentFolder->ChildAssets.Add(TopLevelAssemblyItem);
+
+		FolderItem->ChildAssets.Remove(TopLevelAssemblyItem);
+
+		SchemaToConfigure->DefaultAssemblyPath = TopLevelAssemblyItem->ParentFolder->GetPath();
+	}
+
+	for (const TSharedPtr<FSchemaContentHierarchyAsset>& ChildAsset : FolderItem->ChildAssets)
+	{
+		DeleteAssetItem(ChildAsset);
+	}
+	FolderItem->ChildAssets.Empty();
+
+	for (const TSharedPtr<FSchemaContentHierarchyFolder>& ChildFolder : FolderItem->ChildFolders)
+	{
+		DeleteFolderItem(ChildFolder);
+	}
+	FolderItem->ChildFolders.Empty();
+}
+
+void SCineAssemblySchemaWindow::DeleteAssetItem(TSharedPtr<FSchemaContentHierarchyAsset> AssetItem)
+{
+	if (AssetItem == TopLevelAssemblyItem)
+	{
+		return;
+	}
+
+	if (UMovieSceneSubAssemblySection* Section = AssetItem->SubAssemblySection.Get())
+	{
+		// Find the track that this section belongs to, and remove the section from that track. 
+		// IF the track has no other sections after deleting the section, then remove the whole track.
+		if (UMovieSceneSubAssemblyTrack* Track = Section->GetTypedOuter<UMovieSceneSubAssemblyTrack>())
+		{
+			Track->RemoveSection(*Section);
+			if (Track->GetAllSections().IsEmpty())
+			{
+				SchemaToConfigure->TemplateSequence->GetMovieScene()->RemoveTrack(*Track);
+			}
+		}
+	}
 }
 
 bool SCineAssemblySchemaWindow::ContainsRecursive(TSharedPtr<FSchemaContentHierarchyFolder> InParentFolder, TSharedPtr<FSchemaContentHierarchyAsset> InAssetItem)
@@ -1450,7 +1568,7 @@ bool SCineAssemblySchemaWindow::ContainsRecursive(TSharedPtr<FSchemaContentHiera
 	}
 
 	bool bFoundRecursively = false;
-	for (TSharedPtr<FSchemaContentHierarchyFolder> ChildFolder : InParentFolder->ChildFolders)
+	for (const TSharedPtr<FSchemaContentHierarchyFolder>& ChildFolder : InParentFolder->ChildFolders)
 	{
 		bFoundRecursively |= ContainsRecursive(ChildFolder, InAssetItem);
 	}
@@ -1528,74 +1646,96 @@ void SCineAssemblySchemaWindow::UpdateFolderList()
 	}
 }
 
-void SCineAssemblySchemaWindow::UpdateAssetList()
+void SCineAssemblySchemaWindow::OnMovieSceneDataChanged(EMovieSceneDataChangeType InChangeType)
 {
-	// Update the schema's asset list and templates
-	const TArray<FString> CachedAssetList = SchemaToConfigure->SubsequencesToCreate;
-	const TMap<FString, FSoftObjectPath> CachedSubsequenceTemplates = SchemaToConfigure->SubsequenceTemplates;
-
-	TArray<FString> AssetList;
-	TMap<FString, FSoftObjectPath> Templates;
-	GetAssetListRecursive(RootFolder, AssetList, Templates);
-
-	SchemaToConfigure->DefaultAssemblyPath = TopLevelAssemblyItem->ParentFolder->Path;
-
-	if (AssetList != CachedAssetList)
+	// In order to mute specific tracks, we only care about changes when a track is added or when a track value changes (which includes the user trying to unmute the track)
+	if ((InChangeType == EMovieSceneDataChangeType::MovieSceneStructureItemAdded) 
+		|| (InChangeType == EMovieSceneDataChangeType::MovieSceneStructureItemRemoved)
+		|| (InChangeType == EMovieSceneDataChangeType::TrackValueChanged))
 	{
-		SchemaToConfigure->Modify();
-		SchemaToConfigure->SubsequencesToCreate = AssetList;
-	}
+		if (UMovieScene* MovieScene = SchemaToConfigure->TemplateSequence->GetMovieScene())
+		{
+			// Because we do not have any additional information about what the actual MovieScene change was, we just loop through every track and force all relevant tracks to be muted
+			for (UMovieSceneTrack* Track : MovieScene->GetTracks())
+			{
+				// Mute CVar tracks and Subsequence Tracks to prevent their evaluation from affecting the world
+				if (UMovieSceneCVarTrack* CVarTrack = Cast<UMovieSceneCVarTrack>(Track))
+				{
+					CVarTrack->SetLocalEvalDisabled(true);
+				}
+				else if (UMovieSceneSubAssemblyTrack* SchemaSubTrack = Cast<UMovieSceneSubAssemblyTrack>(Track))
+				{
+					SchemaSubTrack->SetLocalEvalDisabled(true);
+				}
+			}
+		}
 
-	if (SchemaToConfigure->Template != TopLevelAssemblyItem->AssetTemplate)
-	{
-		SchemaToConfigure->Modify();
-		SchemaToConfigure->Template = TopLevelAssemblyItem->AssetTemplate;
+		// It is possible that a new Subsequence or Shot was just added or removed, so we update the subsequence items in the content tree
+		UpdateContentHierarchyAssets();
+
+		TreeView->RequestTreeRefresh();
+		ExpandTreeRecursive(RootFolder);
 	}
+}
 
-	bool bSubsequenceTemplatesHaveChanged = Templates.Num() != CachedSubsequenceTemplates.Num();
+void SCineAssemblySchemaWindow::UpdateContentHierarchyAssets()
+{
+	// Cache the list of asset items that are already in the tree. Items who subsections are no longer in the template movie scene will be removed later.
+	TArray<TSharedPtr<FSchemaContentHierarchyAsset>> AssetItemsNotFound;
+	GetAssetListRecursive(RootFolder, AssetItemsNotFound);
+	AssetItemsNotFound.Remove(TopLevelAssemblyItem);
 
-	if (!bSubsequenceTemplatesHaveChanged)
+	// Get the list of SubAssembly sections in the template movie scene
+	TArray<UMovieSceneSubAssemblySection*> SubAssemblySections;
+	if (UMovieScene* SchemaTemplateMovieScene = SchemaToConfigure->TemplateSequence->GetMovieScene())
 	{
-		for (auto Iter = Templates.CreateConstIterator(); Iter; ++Iter)
+		for (UMovieSceneTrack* Track : SchemaTemplateMovieScene->GetTracks())
 		{
-			if (const FSoftObjectPath* Found = CachedSubsequenceTemplates.Find(Iter->Key);
-				Found == nullptr || (Found != nullptr && *Found != Iter->Value))
+			if (UMovieSceneSubAssemblyTrack* SubAssemblyTrack = Cast<UMovieSceneSubAssemblyTrack>(Track))
 			{
-				bSubsequenceTemplatesHaveChanged = true;
-				break;
+				for (UMovieSceneSection* Section : SubAssemblyTrack->GetAllSections())
+				{
+					if (UMovieSceneSubAssemblySection* SubAssemblySection = Cast<UMovieSceneSubAssemblySection>(Section))
+					{
+						SubAssemblySections.Add(SubAssemblySection);
+					}
+				}
 			}
 		}
 	}
 
-	if (bSubsequenceTemplatesHaveChanged)
-	{
-		SchemaToConfigure->Modify();
-		SchemaToConfigure->SubsequenceTemplates = Templates;
-	}
-}
-
-void SCineAssemblySchemaWindow::OnMovieSceneDataChanged(EMovieSceneDataChangeType InChangeType)
-{
-	// In order to mute specific tracks, we only care about changes when a track is added or when a track value changes (which includes the user trying to unmute the track)
-	if ((InChangeType == EMovieSceneDataChangeType::MovieSceneStructureItemAdded) || (InChangeType == EMovieSceneDataChangeType::TrackValueChanged))
+	for (UMovieSceneSubAssemblySection* SubAssemblySection : SubAssemblySections)
 	{
-		if (UMovieScene* MovieScene = SchemaToConfigure->TemplateSequence->GetMovieScene())
+		if (SubAssemblySection->IsTemplateSection())
 		{
-			// Because we do not have any additional information about what the actual MovieScene change was, we just loop through every track and force all relevant tracks to be muted
-			for (UMovieSceneTrack* Track : MovieScene->GetTracks())
+			// Walk the tree until we find a parent folder whose path matches the path of the SubAssembly section
+			if (TSharedPtr<FSchemaContentHierarchyFolder> ParentFolder = FindFolderByPathRecursive(RootFolder, SubAssemblySection->GetSequencePath()))
 			{
-				// Mute CVar tracks and Subsequence Tracks to prevent their evaluation from affecting the world
-				if (UMovieSceneCVarTrack* CVarTrack = Cast<UMovieSceneCVarTrack>(Track))
+				// Try to find a child asset item that already represents this SubAssembly section
+				TSharedPtr<FSchemaContentHierarchyAsset>* AssetItemPtr = ParentFolder->ChildAssets.FindByPredicate([SubAssemblySection](const TSharedPtr<FSchemaContentHierarchyAsset>& AssetItem) { return AssetItem->SubAssemblySection == SubAssemblySection; });
+				if (AssetItemPtr)
 				{
-					CVarTrack->SetLocalEvalDisabled(true);
+					// If the asset item exists, then it is still valid (it matches one of the existing SubAssembly sections)
+					AssetItemsNotFound.Remove(*AssetItemPtr);
 				}
-				else if (UMovieSceneSubAssemblyTrack* SchemaSubTrack = Cast<UMovieSceneSubAssemblyTrack>(Track))
+				else
 				{
-					SchemaSubTrack->SetLocalEvalDisabled(true);
+					// If the asset item does not yet exist, make a new one and add it into the tree
+					TSharedPtr<FSchemaContentHierarchyAsset> NewItem = MakeShared<FSchemaContentHierarchyAsset>();
+					NewItem->ParentFolder = ParentFolder;
+					NewItem->SubAssemblySection = SubAssemblySection;
+					ParentFolder->ChildAssets.Add(NewItem);
 				}
 			}
 		}
 	}
+
+	// Any asset items that didn't match current SubAssembly sections need to be removed from the tree. 
+	// This could happen if a SubAssembly section is removed from the Schema Sequencer, for example.
+	for (const TSharedPtr<FSchemaContentHierarchyAsset>& AssetItem : AssetItemsNotFound)
+	{
+		AssetItem->ParentFolder->ChildAssets.Remove(AssetItem);
+	}
 }
 
 void SCineAssemblySchemaWindow::OnActorAddedToSequencer(AActor* InActor, const FGuid InBinding)
@@ -1724,8 +1864,8 @@ FReply SCineAssemblySchemaWindow::OnTreeRowDragDetected(const FGeometry& MyGeome
 				return FReply::Handled();
 			}
 			
-			FSoftObjectPath TemplatePath = SelectedTreeItem->AsAsset() ? SelectedTreeItem->AsAsset()->AssetTemplate : FSoftObjectPath();
-			TSharedRef<FSchemaAssetDragDrop> Operation = FSchemaAssetDragDrop::New(FSchemaAssetParameters(FPaths::GetPathLeaf(SelectedTreeItem->Path), TemplatePath));
+			FSoftObjectPath TemplatePath = SelectedTreeItem->AsAsset() ? SelectedTreeItem->AsAsset()->GetTemplateObjectPath() : FSoftObjectPath();
+			TSharedRef<FSchemaAssetDragDrop> Operation = FSchemaAssetDragDrop::New(FSchemaAssetParameters(SelectedTreeItem->GetName(), TemplatePath));
 			Operation->SourceTreeItem = SelectedTreeItem;
 
 			if (TSharedPtr<FSchemaContentHierarchyAsset> AssetItem = SelectedTreeItem->AsAsset())
@@ -1784,37 +1924,43 @@ FReply SCineAssemblySchemaWindow::OnTreeRowAcceptDrop(const FDragDropEvent& InDr
 			DroppedItem = MakeShared<FSchemaContentHierarchyAsset>();
 		}
 
-		// If an item with the same name already exists in the drop location, do not handle the drop (the dragged item will be reset to its original location)
-		FString PendingDroppedItemPath = TargetFolder->Path / Operation->ItemName;
-		if (DroppedItem->AsAsset() && Algo::ContainsBy(TargetFolder->ChildAssets, PendingDroppedItemPath, &ISchemaContentHierarchyTreeItem::Path))
-		{
-			return FReply::Unhandled();
-		}
-		else if (DroppedItem->AsFolder() && Algo::ContainsBy(TargetFolder->ChildFolders, PendingDroppedItemPath, &ISchemaContentHierarchyTreeItem::Path))
+		// If the dropped item is a folder, but the target folder already contains a subfolder with the same name, do not handle the drop (the dragged item will be reset to its original location)
+		const FString DroppedItemPath = TargetFolder->GetPath() / Operation->ItemName;
+		if (DroppedItem->AsFolder())
 		{
-			return FReply::Unhandled();
+			if (Algo::AnyOf(TargetFolder->ChildFolders, [DroppedItemPath](const TSharedPtr<FSchemaContentHierarchyFolder>& FolderPtr) { return FolderPtr && FolderPtr->GetPath() == DroppedItemPath; }))
+			{
+				return FReply::Unhandled();
+			}
 		}
 
-		DroppedItem->Path = PendingDroppedItemPath;
+		DroppedItem->SetPath(DroppedItemPath);
 		DroppedItem->ParentFolder = TargetFolder;
 
 		if (TSharedPtr<FSchemaContentHierarchyAsset> DroppedAsset = DroppedItem->AsAsset())
 		{
-			DroppedAsset->AssetTemplate = Operation->Template;
-
 			TargetFolder->ChildAssets.Add(DroppedAsset);
 			TargetFolder->ChildAssets.Sort(UE::CineAssemblyTools::Private::SortTreeItems);
+
+			if (DroppedAsset == TopLevelAssemblyItem)
+			{
+				SchemaToConfigure->DefaultAssemblyPath = TopLevelAssemblyItem->ParentFolder->GetPath();
+			}
 		}
 		else if (TSharedPtr<FSchemaContentHierarchyFolder> DroppedFolder = DroppedItem->AsFolder())
 		{
 			TargetFolder->ChildFolders.Add(DroppedFolder);
 
-			SetChildrenPathRecursive(DroppedFolder, DroppedFolder->Path);
+			SetChildrenPathRecursive(DroppedFolder, DroppedFolder->GetPath());
 			TargetFolder->ChildFolders.Sort(UE::CineAssemblyTools::Private::SortTreeItems);
+
+			if (ContainsRecursive(DroppedFolder, TopLevelAssemblyItem))
+			{
+				SchemaToConfigure->DefaultAssemblyPath = TopLevelAssemblyItem->ParentFolder->GetPath();
+			}
 		}
 
 		UpdateFolderList();
-		UpdateAssetList();
 
 		ExpandTreeRecursive(TargetFolder);
 		TreeView->RequestTreeRefresh();
diff --git a/Engine/Plugins/VirtualProduction/CinematicAssemblyTools/Source/CineAssemblyToolsEditor/Private/UI/CineAssembly/SCineAssemblySchemaWindow.h b/Engine/Plugins/VirtualProduction/CinematicAssemblyTools/Source/CineAssemblyToolsEditor/Private/UI/CineAssembly/SCineAssemblySchemaWindow.h
--- a/Engine/Plugins/VirtualProduction/CinematicAssemblyTools/Source/CineAssemblyToolsEditor/Private/UI/CineAssembly/SCineAssemblySchemaWindow.h
+++ b/Engine/Plugins/VirtualProduction/CinematicAssemblyTools/Source/CineAssemblyToolsEditor/Private/UI/CineAssembly/SCineAssemblySchemaWindow.h
@@ -173,6 +173,12 @@ private:
 	/** Gets the children of the input tree view item to build additional tree rows */
 	void OnGetChildren(TSharedPtr<ISchemaContentHierarchyTreeItem> TreeItem, TArray<TSharedPtr<ISchemaContentHierarchyTreeItem>>& OutNodes);
 
+	/** Callback to get the display name of the input tree item */
+	FText GetTreeItemDisplayName(TSharedPtr<ISchemaContentHierarchyTreeItem> TreeItem) const;
+
+	/** Callback to get the Template object path of the input tree item */
+	FString GetTemplateObjectPath(TSharedPtr<FSchemaContentHierarchyAsset> AssetItem) const;
+
 	/** Callback when the content tree view rebuilds itself, used to make the textbox in the most recently added row editable */
 	void OnTreeItemsRebuilt();
 
@@ -194,8 +200,8 @@ private:
 	/** Recursively get the path of every folder item in the tree view */
 	void GetFolderListRecursive(TSharedPtr<FSchemaContentHierarchyFolder> FolderItem, TArray<FString>& OutFolderList) const;
 
-	/** Recursively get the path of every asset item in the tree view */
-	void GetAssetListRecursive(TSharedPtr<FSchemaContentHierarchyFolder> FolderItem, TArray<FString>& OutAssetList, TMap<FString, FSoftObjectPath>& OutTemplates) const;
+	/** Recursively find every asset item in the content hierarchy tree */
+	void GetAssetListRecursive(TSharedPtr<FSchemaContentHierarchyFolder> FolderItem, TArray<TSharedPtr<FSchemaContentHierarchyAsset>>& OutAssetList) const;
 
 	/** Returns the folder in the tree whose path matches the input path */
 	TSharedPtr<FSchemaContentHierarchyFolder> FindFolderByPathRecursive(TSharedPtr<FSchemaContentHierarchyFolder> FolderItem, const FString& FolderPath) const;
@@ -210,7 +216,7 @@ private:
 	void OnTreeItemTextCommitted(const FText& InText, ETextCommit::Type InCommitType, TSharedPtr<ISchemaContentHierarchyTreeItem> TreeItem);
 
 	/** Updates the Template object of the given tree item. */
-	void OnTreeItemTemplateChanged(const FAssetData& InAssetData, TSharedPtr<ISchemaContentHierarchyTreeItem> TreeItem);
+	void OnTreeItemTemplateChanged(const FAssetData& InAssetData, TSharedPtr<FSchemaContentHierarchyAsset> AssetItem);
 
 	/** Callback when one of the properties of the schema being configured changes */
 	void OnSchemaPropertiesChanged(const FPropertyChangedEvent& PropertyChangedEvent);
@@ -227,6 +233,12 @@ private:
 	/** Remove the input item and its children from the tree view */
 	void DeleteTreeItem(TSharedPtr<ISchemaContentHierarchyTreeItem> TreeItem);
 
+	/** Removes an asset item from the tree view, and from the template sequence */
+	void DeleteAssetItem(TSharedPtr<FSchemaContentHierarchyAsset> AssetItem);
+
+	/** Removes a folder and all of its children from the tree view */
+	void DeleteFolderItem(TSharedPtr<FSchemaContentHierarchyFolder> FolderItem);
+
 	/** Begins a drag and drop event to drag an item out of the asset list view */
 	FReply OnAssetRowDragDetected(const FGeometry& MyGeometry, const FPointerEvent& MouseEvent);
 
@@ -245,15 +257,15 @@ private:
 	/** Update the list of folders to create from the current set in the tree view */
 	void UpdateFolderList();
 
-	/** Update the list of assets to create from the current set in the tree view */
-	void UpdateAssetList();
-
 	/** Callback when an actor is added to the template sequence, used to remove the actor(s) from the binding */
 	void OnActorAddedToSequencer(AActor* InActor, const FGuid InBinding);
 
 	/** Callback whenever MovieScene data in the template sequence is changed */
 	void OnMovieSceneDataChanged(EMovieSceneDataChangeType InChangeType);
 
+	/** Add/delete asset items from the content hierarchy tree based on the SubAssembly sections in the TemplateSequence MovieScene */
+	void UpdateContentHierarchyAssets();
+
 	/** Extends the sequencer widget's toolbar */
 	void AddSequencerToolbarExtension(FToolBarBuilder& Builder);
 
