ShotManagement: Fix Schema Sequencer to support assets dropped directly onto a SubAssembly Track.
#rnx

#Synced-CL 51337922

diff --git a/Engine/Plugins/VirtualProduction/CinematicAssemblyTools/Source/CineAssemblyToolsEditor/Private/Sequencer/SubAssemblyTrackEditor.cpp b/Engine/Plugins/VirtualProduction/CinematicAssemblyTools/Source/CineAssemblyToolsEditor/Private/Sequencer/SubAssemblyTrackEditor.cpp
--- a/Engine/Plugins/VirtualProduction/CinematicAssemblyTools/Source/CineAssemblyToolsEditor/Private/Sequencer/SubAssemblyTrackEditor.cpp
+++ b/Engine/Plugins/VirtualProduction/CinematicAssemblyTools/Source/CineAssemblyToolsEditor/Private/Sequencer/SubAssemblyTrackEditor.cpp
@@ -5,6 +5,7 @@
 #include "CineAssembly.h"
 #include "CineAssemblySchema.h"
 #include "ContentBrowserModule.h"
+#include "DragAndDrop/AssetDragDropOp.h"
 #include "IContentBrowserSingleton.h"
 #include "MovieSceneSubAssemblySection.h"
 #include "MovieSceneSubAssemblyTrack.h"
@@ -56,11 +57,54 @@ bool FSubAssemblyTrackEditor::SupportsSequence(UMovieSceneSequence* InSequence)
 }
 
 bool FSubAssemblyTrackEditor::HandleAssetAdded(UObject* Asset, const FGuid& TargetObjectGuid)
+{
+	if (UMovieScene* FocusedMovieScene = GetFocusedMovieScene())
+	{
+		const FFrameNumber FrameNumber = GetSequencer()->GetLocalTime().Time.FrameNumber;
+		
+		// If the drop is successful, a new track will be created and the dropped asset will be added to it
+		UMovieSceneSubAssemblyTrack* NullTrack = nullptr;
+		return AddDroppedAsset(Asset, NullTrack, FrameNumber);
+	}
+	return false;
+}
+
+FReply FSubAssemblyTrackEditor::OnDrop(const FDragDropEvent& DragDropEvent, const FSequencerDragDropParams& DragDropParams)
+{
+	if (!DragDropParams.Track.IsValid())
+	{
+		return FReply::Unhandled();
+	}
+
+	if (!DragDropParams.Track.Get()->IsA(GetSubTrackClass()))
+	{
+		return FReply::Unhandled();
+	}
+
+	TSharedPtr<FDragDropOperation> Operation = DragDropEvent.GetOperation();
+
+	if (!Operation.IsValid() || !Operation->IsOfType<FAssetDragDropOp>())
+	{
+		return FReply::Unhandled();
+	}
+
+	TSharedPtr<FAssetDragDropOp> DragDropOp = StaticCastSharedPtr<FAssetDragDropOp>(Operation);
+
+	bool bAnyDropped = false;
+	for (const FAssetData& AssetData : DragDropOp->GetAssets())
+	{
+		bAnyDropped |= AddDroppedAsset(AssetData.GetAsset(), Cast<UMovieSceneSubAssemblyTrack>(DragDropParams.Track), DragDropParams.FrameNumber);
+	}
+
+	return bAnyDropped ? FReply::Handled() : FReply::Unhandled();
+}
+
+bool FSubAssemblyTrackEditor::AddDroppedAsset(UObject* Asset, UMovieSceneSubAssemblyTrack* Track, FFrameNumber FrameNumber)
 {
 	// If the dropped asset is a Schema, it can only be added as a template
-	if (Asset->IsA(UCineAssemblySchema::StaticClass()))
+	if (Asset && Asset->IsA(UCineAssemblySchema::StaticClass()))
 	{
-		OnDroppedAsTemplate(Asset);
+		OnDroppedAsTemplate(Asset, Track, FrameNumber);
 		return true;
 	}
 
@@ -89,14 +133,14 @@ bool FSubAssemblyTrackEditor::HandleAssetAdded(UObject* Asset, const FGuid& Targ
 		LOCTEXT("OnDroppedNewSequence", "Add New Sequence"),
 		LOCTEXT("OnDroppedNewSequenceTooltip", "Add a new Subsequence using the dropped Sequence as the base"),
 		FSlateIcon(),
-		FUIAction(FExecuteAction::CreateRaw(this, &FSubAssemblyTrackEditor::OnDroppedAsTemplate, Asset))
+		FUIAction(FExecuteAction::CreateRaw(this, &FSubAssemblyTrackEditor::OnDroppedAsTemplate, Asset, Track, FrameNumber))
 	);
 
 	MenuBuilder.AddMenuEntry(
 		LOCTEXT("OnDroppedReference", "Add Reference to Sequence"),
 		LOCTEXT("OnDroppedReferenceTooltip", "Add a reference to the dropped Sequence"),
 		FSlateIcon(),
-		FUIAction(FExecuteAction::CreateRaw(this, &FSubAssemblyTrackEditor::OnDroppedAsReference, Sequence))
+		FUIAction(FExecuteAction::CreateRaw(this, &FSubAssemblyTrackEditor::OnDroppedAsReference, Sequence, Track, FrameNumber))
 	);
 
 	// Open the menu where at the user's current mouse position
@@ -356,7 +400,7 @@ void FSubAssemblyTrackEditor::AddNewSequenceToTrack(UMovieSceneSubAssemblyTrack*
 		return;
 	}
 
-	const FFrameNumber StartFrame = FocusedMovieScene->GetPlaybackRange().GetLowerBoundValue();
+	const FFrameNumber StartFrame = GetTimeForKey();
 	int32 Duration = FocusedMovieScene->GetPlaybackRange().Size<FFrameNumber>().Value;
 
 	// If the template object is a valid sequence, use its duration instead of the default duration
@@ -402,26 +446,50 @@ void FSubAssemblyTrackEditor::OnAssetEnterPressedAsReference(const TArray<FAsset
 	}
 }
 
-void FSubAssemblyTrackEditor::OnDroppedAsTemplate(UObject* Asset)
+void FSubAssemblyTrackEditor::OnDroppedAsTemplate(UObject* Asset, UMovieSceneSubAssemblyTrack* Track, FFrameNumber FrameNumber)
 {
-	// Determine the appropriate track type for the dropped asset
-	ESubAssemblyTrackType TrackType = ESubAssemblyTrackType::SubsequenceTrack;
-	if (UMovieSceneSequence* Sequence = Cast<UMovieSceneSequence>(Asset))
+	// If the input track was null, create a new one
+	if (!Track)
 	{
-		TrackType = (Sequence->GetMovieScene()->GetCameraCutTrack()) ? ESubAssemblyTrackType::CinematicShotTrack : ESubAssemblyTrackType::SubsequenceTrack;
+		const ESubAssemblyTrackType TrackType = GetTrackTypeForAsset(Asset);
+		Track = AddSubAssemblyTrack(TrackType);
 	}
 
-	UMovieSceneSubAssemblyTrack* NewTrack = AddSubAssemblyTrack(TrackType);
-	AddNewSequenceToTrack(NewTrack, Asset);
+	// Initiate keying at the input frame number to ensure that this is where the dropped sequence will be placed
+	FMovieSceneTrackEditor::BeginKeying(FrameNumber);
+	AddNewSequenceToTrack(Track, Asset);
+	FMovieSceneTrackEditor::EndKeying();
 }
 
-void FSubAssemblyTrackEditor::OnDroppedAsReference(UMovieSceneSequence* Sequence)
+void FSubAssemblyTrackEditor::OnDroppedAsReference(UMovieSceneSequence* Sequence, UMovieSceneSubAssemblyTrack* Track, FFrameNumber FrameNumber)
 {
-	// Determine the appropriate track type for the dropped sequence
-	ESubAssemblyTrackType TrackType = (Sequence->GetMovieScene()->GetCameraCutTrack()) ? ESubAssemblyTrackType::CinematicShotTrack : ESubAssemblyTrackType::SubsequenceTrack;
+	// If the input track was null, create a new one
+	if (!Track)
+	{
+		const ESubAssemblyTrackType TrackType = GetTrackTypeForAsset(Sequence);
+		Track = AddSubAssemblyTrack(TrackType);
+	}
+
+	// Initiate keying at the input frame number to ensure that this is where the dropped sequence will be placed
+	FMovieSceneTrackEditor::BeginKeying(FrameNumber);
+	FSubTrackEditor::AddSequenceToTrack(Sequence, Track);
+	FMovieSceneTrackEditor::EndKeying();
+}
+
+ESubAssemblyTrackType FSubAssemblyTrackEditor::GetTrackTypeForAsset(UObject* Asset)
+{
+	if (UMovieSceneSequence* Sequence = Cast<UMovieSceneSequence>(Asset))
+	{
+		if (UMovieScene* MovieScene = Sequence->GetMovieScene())
+		{
+			if (MovieScene->GetCameraCutTrack())
+			{
+				return ESubAssemblyTrackType::CinematicShotTrack;
+			}
+		}
+	}
 
-	UMovieSceneSubAssemblyTrack* NewTrack = AddSubAssemblyTrack(TrackType);
-	FSubTrackEditor::AddSequenceToTrack(Sequence, NewTrack);
+	return ESubAssemblyTrackType::SubsequenceTrack;
 }
 
 #undef LOCTEXT_NAMESPACE
diff --git a/Engine/Plugins/VirtualProduction/CinematicAssemblyTools/Source/CineAssemblyToolsEditor/Private/Sequencer/SubAssemblyTrackEditor.h b/Engine/Plugins/VirtualProduction/CinematicAssemblyTools/Source/CineAssemblyToolsEditor/Private/Sequencer/SubAssemblyTrackEditor.h
--- a/Engine/Plugins/VirtualProduction/CinematicAssemblyTools/Source/CineAssemblyToolsEditor/Private/Sequencer/SubAssemblyTrackEditor.h
+++ b/Engine/Plugins/VirtualProduction/CinematicAssemblyTools/Source/CineAssemblyToolsEditor/Private/Sequencer/SubAssemblyTrackEditor.h
@@ -29,6 +29,7 @@ public:
 	virtual TSharedPtr<SWidget> BuildOutlinerEditWidget(const FGuid& ObjectBinding, UMovieSceneTrack* Track, const FBuildEditWidgetParams& Params) override;
 	virtual TSharedPtr<SWidget> BuildOutlinerColumnWidget(const FBuildColumnWidgetParams& Params, const FName& ColumnName) override;
 	virtual bool HandleAssetAdded(UObject* Asset, const FGuid& TargetObjectGuid) override;
+	virtual FReply OnDrop(const FDragDropEvent& DragDropEvent, const FSequencerDragDropParams& DragDropParams) override;
 	virtual TSharedRef<ISequencerSection> MakeSectionInterface(UMovieSceneSection& SectionObject, UMovieSceneTrack& Track, FGuid ObjectBinding) override;
 	virtual bool SupportsSequence(UMovieSceneSequence* InSequence) const override;
 	//~ End ISequencerTrackEditor interface
@@ -68,9 +69,15 @@ private:
 	/** Creates and adds a new sequence to the input Track, then finds the newly created SubAssemblySection and assigns the input TemplateObject to it */
 	void AddNewSequenceToTrack(UMovieSceneSubAssemblyTrack* Track, UObject* TemplateObject);
 
+	/** Adds the input Asset onto the input Track (or a new Track if that is null) at the input FrameNumber */
+	bool AddDroppedAsset(UObject* Asset, UMovieSceneSubAssemblyTrack* Track, FFrameNumber FrameNumber);
+
 	/** Callback when an asset is dropped and handled by this track editor to be added as a template for a new sequence */
-	void OnDroppedAsTemplate(UObject* Asset);
+	void OnDroppedAsTemplate(UObject* Asset, UMovieSceneSubAssemblyTrack* Track, FFrameNumber FrameNumber);
 
 	/** Callback when an asset is dropped and handled by this track editor to be added as a reference */
-	void OnDroppedAsReference(UMovieSceneSequence* Sequence);
+	void OnDroppedAsReference(UMovieSceneSequence* Sequence, UMovieSceneSubAssemblyTrack* Track, FFrameNumber FrameNumber);
+
+	/** Return the appropriate track type for the input asset */
+	ESubAssemblyTrackType GetTrackTypeForAsset(UObject* Asset);
 };
\ No newline at end of file
