ShotManagement: Add new Cine Assembly Editor Function Library to expose relevant functionality to editor scriptng.
Add two functions to that library, OpenAssembly, and OpenAssociatedLevel.
#rnx

#Synced-CL 51337922

diff --git a/Engine/Plugins/VirtualProduction/CinematicAssemblyTools/Source/CineAssemblyToolsEditor/CineAssemblyToolsEditor.build.cs b/Engine/Plugins/VirtualProduction/CinematicAssemblyTools/Source/CineAssemblyToolsEditor/CineAssemblyToolsEditor.build.cs
--- a/Engine/Plugins/VirtualProduction/CinematicAssemblyTools/Source/CineAssemblyToolsEditor/CineAssemblyToolsEditor.build.cs
+++ b/Engine/Plugins/VirtualProduction/CinematicAssemblyTools/Source/CineAssemblyToolsEditor/CineAssemblyToolsEditor.build.cs
@@ -28,6 +28,7 @@ namespace UnrealBuildTool.Rules
 					"DesktopPlatform",
 					"DeveloperSettings",
 					"DirectoryPlaceholder",
+					"EditorSubsystem",
 					"EditorWidgets",
 					"Engine",
 					"InputCore",
diff --git a/Engine/Plugins/VirtualProduction/CinematicAssemblyTools/Source/CineAssemblyToolsEditor/Private/AssetDefinition_CineAssembly.cpp b/Engine/Plugins/VirtualProduction/CinematicAssemblyTools/Source/CineAssemblyToolsEditor/Private/AssetDefinition_CineAssembly.cpp
--- a/Engine/Plugins/VirtualProduction/CinematicAssemblyTools/Source/CineAssemblyToolsEditor/Private/AssetDefinition_CineAssembly.cpp
+++ b/Engine/Plugins/VirtualProduction/CinematicAssemblyTools/Source/CineAssemblyToolsEditor/Private/AssetDefinition_CineAssembly.cpp
@@ -4,9 +4,7 @@
 #include "AssetToolsModule.h"
 #include "CineAssembly.h"
 #include "CineAssemblyToolsStyle.h"
-
-#include "Editor.h"
-#include "FileHelpers.h"
+#include "CineAssemblyEditorFunctionLibrary.h"
 
 #define LOCTEXT_NAMESPACE "AssetTypeActions"
 
@@ -125,49 +123,13 @@ TArray<FAssetData> UAssetDefinition_CineAssembly::PrepareToActivateAssets(const
 	if (ActivateArgs.Assets.Num() == 1)
 	{
 		const FAssetData& CineAssemblyData = ActivateArgs.Assets[0];
-		AssetsToOpen.Add(CineAssemblyData);
-
-		const FName IsDataOnlyAssetRegistrySearchablePropertyName = GET_MEMBER_NAME_CHECKED(UCineAssembly, bIsDataOnly);
-
-		bool bAssemblyIsDataOnly;
-		CineAssemblyData.GetTagValue(IsDataOnlyAssetRegistrySearchablePropertyName, bAssemblyIsDataOnly);
 
 		UCineAssembly* CineAssembly = Cast<UCineAssembly>(CineAssemblyData.GetAsset());
-		if (!bAssemblyIsDataOnly && CineAssembly != nullptr)
+
+		// Attempt to open the level associated with this assembly (if it has one). If this operation fails, then the assembly should not be opened. 
+		if (UCineAssemblyEditorFunctionLibrary::OpenAssociatedLevel(CineAssembly))
 		{
-			if (CineAssembly->Level.IsValid())
-			{
-				if (UWorld* WorldToOpen = Cast<UWorld>(CineAssembly->Level.TryLoad()))
-				{
-					UWorld* CurrentWorld = GEditor->GetEditorWorldContext().World();
-					if (CurrentWorld != WorldToOpen)
-					{
-						// Prompt the user to save their unsaved changes to the current level before loading the level associated with this asset
-						constexpr bool bPromptUserToSave = true;
-						constexpr bool bSaveMapPackages = true;
-						constexpr bool bSaveContentPackages = false;
-						if (FEditorFileUtils::SaveDirtyPackages(bPromptUserToSave, bSaveMapPackages, bSaveContentPackages))
-						{
-							if (!WorldToOpen->GetPackage()->HasAnyPackageFlags(PKG_NewlyCreated))
-							{
-								const FString FileToOpen = FPackageName::LongPackageNameToFilename(WorldToOpen->GetOutermost()->GetName(), FPackageName::GetMapPackageExtension());
-								const bool bLoadAsTemplate = false;
-								const bool bShowProgress = true;
-								FEditorFileUtils::LoadMap(FileToOpen, bLoadAsTemplate, bShowProgress);
-							}
-						}
-						else
-						{
-							// If the user canceled out of the prompt to save the current level, then do not try to open the asset
-							AssetsToOpen.Empty();
-						}
-					}
-				}
-				else
-				{
-					UE_LOG(LogCineAssemblyDefinition, Error, TEXT("Failed to load %s while opening %s"), *CineAssembly->Level.ToString(), *CineAssembly->GetFName().ToString());
-				}
-			}
+			AssetsToOpen.Add(CineAssemblyData);
 		}
 	}
 
diff --git a/Engine/Plugins/VirtualProduction/CinematicAssemblyTools/Source/CineAssemblyToolsEditor/Private/CineAssemblyEditorFunctionLibrary.cpp b/Engine/Plugins/VirtualProduction/CinematicAssemblyTools/Source/CineAssemblyToolsEditor/Private/CineAssemblyEditorFunctionLibrary.cpp
new file mode 100644
--- /dev/null
+++ b/Engine/Plugins/VirtualProduction/CinematicAssemblyTools/Source/CineAssemblyToolsEditor/Private/CineAssemblyEditorFunctionLibrary.cpp
@@ -0,0 +1,61 @@
+// Copyright Epic Games, Inc. All Rights Reserved.
+
+#include "CineAssemblyEditorFunctionLibrary.h"
+
+#include "CineAssembly.h"
+#include "Editor.h"
+#include "FileHelpers.h"
+#include "Subsystems/AssetEditorSubsystem.h"
+
+void UCineAssemblyEditorFunctionLibrary::OpenAssembly(UCineAssembly* CineAssembly, bool bOpenAssociatedLevel, FCineAssemblyLevelSaveOptions SaveOptions)
+{
+	if (bOpenAssociatedLevel)
+	{
+		OpenAssociatedLevel(CineAssembly, SaveOptions);
+	}
+
+	GEditor->GetEditorSubsystem<UAssetEditorSubsystem>()->OpenEditorForAssets({ CineAssembly });
+}
+
+bool UCineAssemblyEditorFunctionLibrary::OpenAssociatedLevel(UCineAssembly* CineAssembly, FCineAssemblyLevelSaveOptions SaveOptions)
+{
+	if (!CineAssembly)
+	{
+		return false;
+	}
+
+	if (CineAssembly->Level.IsValid() && !CineAssembly->bIsDataOnly)
+	{
+		// Load the level associated with the input assembly
+		if (UWorld* WorldToOpen = Cast<UWorld>(CineAssembly->Level.TryLoad()))
+		{
+			if (SaveOptions.bSaveCurrentLevel)
+			{
+				// Save the currently open level before opening the associated level (if they differ)
+				UWorld* CurrentWorld = GEditor->GetEditorWorldContext().World();
+				if (CurrentWorld != WorldToOpen)
+				{
+					constexpr bool bSaveMapPackages = true;
+					constexpr bool bSaveContentPackages = false;
+					if (!FEditorFileUtils::SaveDirtyPackages(SaveOptions.bPromptForUnsavedChanges, bSaveMapPackages, bSaveContentPackages))
+					{
+						// If the user was prompted to save but cancelled, then do not proceed with opening the level
+						return false;
+					}
+				}
+			}
+
+			// Open the associated level
+			if (!WorldToOpen->GetPackage()->HasAnyPackageFlags(PKG_NewlyCreated))
+			{
+				const FString FileToOpen = FPackageName::LongPackageNameToFilename(WorldToOpen->GetOutermost()->GetName(), FPackageName::GetMapPackageExtension());
+
+				const bool bLoadAsTemplate = false;
+				const bool bShowProgress = true;
+				FEditorFileUtils::LoadMap(FileToOpen, bLoadAsTemplate, bShowProgress);
+			}
+		}
+	}
+
+	return true;
+}
diff --git a/Engine/Plugins/VirtualProduction/CinematicAssemblyTools/Source/CineAssemblyToolsEditor/Public/CineAssemblyEditorFunctionLibrary.h b/Engine/Plugins/VirtualProduction/CinematicAssemblyTools/Source/CineAssemblyToolsEditor/Public/CineAssemblyEditorFunctionLibrary.h
new file mode 100644
--- /dev/null
+++ b/Engine/Plugins/VirtualProduction/CinematicAssemblyTools/Source/CineAssemblyToolsEditor/Public/CineAssemblyEditorFunctionLibrary.h
@@ -0,0 +1,47 @@
+// Copyright Epic Games, Inc. All Rights Reserved.
+
+#pragma once
+
+#include "Kismet/BlueprintFunctionLibrary.h"
+
+#include "CineAssemblyEditorFunctionLibrary.generated.h"
+
+class UCineAssembly;
+
+/**
+ * When opening a level, these options tell us how to handle unsaved changes in the currently open level
+ */
+USTRUCT(BlueprintType)
+struct FCineAssemblyLevelSaveOptions
+{
+	GENERATED_BODY()
+
+	/** Whether the currently open level should be saved before opening a different level */
+	UPROPERTY(BlueprintReadWrite, Category = CineAssembly)
+	bool bSaveCurrentLevel = true;
+
+	/** Whether the user should be prompted before saving the current level */
+	UPROPERTY(BlueprintReadWrite, Category = CineAssembly, meta = (EditCondition = bSaveCurrentLevel))
+	bool bPromptForUnsavedChanges = true;
+};
+
+/**
+ * Library of functions that expose CineAssembly functionality to editor scripting
+ */
+UCLASS()
+class UCineAssemblyEditorFunctionLibrary : public UBlueprintFunctionLibrary
+{
+	GENERATED_BODY()
+
+public:
+	/** Opens a CineAssembly in Sequencer, and optionally opens its associated Level */
+	UFUNCTION(BlueprintCallable, Category = CineAssembly, meta = (AdvancedDisplay = "SaveOptions"))
+	static void OpenAssembly(UCineAssembly* CineAssembly, bool bOpenAssociatedLevel = true, FCineAssemblyLevelSaveOptions SaveOptions = FCineAssemblyLevelSaveOptions());
+
+	/**
+	 * Opens the Level associated with the input CineAssembly (if it has one).
+	 * Returns false only if the input assembly was invalid, or if its associated level was valid but could not be opened.
+	 */
+	UFUNCTION(BlueprintCallable, Category = CineAssembly, meta = (AdvancedDisplay = "SaveOptions"))
+	static bool OpenAssociatedLevel(UCineAssembly* CineAssembly, FCineAssemblyLevelSaveOptions SaveOptions = FCineAssemblyLevelSaveOptions());
+};
