ShotManagement: Remove Asset list view from the CIneAssembly Schema content tab
#rnx

#Synced-CL 51337922

diff --git a/Engine/Plugins/VirtualProduction/CinematicAssemblyTools/Source/CineAssemblyToolsEditor/Private/UI/CineAssembly/SCineAssemblySchemaWindow.cpp b/Engine/Plugins/VirtualProduction/CinematicAssemblyTools/Source/CineAssemblyToolsEditor/Private/UI/CineAssembly/SCineAssemblySchemaWindow.cpp
--- a/Engine/Plugins/VirtualProduction/CinematicAssemblyTools/Source/CineAssemblyToolsEditor/Private/UI/CineAssembly/SCineAssemblySchemaWindow.cpp
+++ b/Engine/Plugins/VirtualProduction/CinematicAssemblyTools/Source/CineAssemblyToolsEditor/Private/UI/CineAssembly/SCineAssemblySchemaWindow.cpp
@@ -455,15 +455,15 @@ TSharedRef<SWidget> SCineAssemblySchemaWindow::MakeMenuPanel()
 				.HAlign(HAlign_Left)
 
 				+ SSegmentedControl<int32>::Slot(2)
-				.Text(LOCTEXT("HierarchyTab", "Content Hierarchy"))
-				.Icon(FSlateIcon(FAppStyle::GetAppStyleSetName(), "Icons.FolderClosed").GetIcon())
-				.HAlign(HAlign_Left)
-
-				+ SSegmentedControl<int32>::Slot(3)
 				.Text(LOCTEXT("TimelineTab", "Timeline Template"))
 				.Icon(FSlateIcon(FCineAssemblyToolsStyle::StyleName, "Icons.Sequencer").GetIcon())
 				.HAlign(HAlign_Left)
 				.Visibility_Lambda([this]() { return SchemaToConfigure->bIsDataOnly ? EVisibility::Collapsed : EVisibility::Visible; })
+
+				+ SSegmentedControl<int32>::Slot(3)
+				.Text(LOCTEXT("HierarchyTab", "Content Hierarchy"))
+				.Icon(FSlateIcon(FAppStyle::GetAppStyleSetName(), "Icons.FolderClosed").GetIcon())
+				.HAlign(HAlign_Left)
 		];
 }
 
@@ -482,12 +482,12 @@ TSharedRef<SWidget> SCineAssemblySchemaWindow::MakeContentPanel()
 
 		+ SWidgetSwitcher::Slot()
 		[
-			MakeHierarchyTabContent()
+			Sequencer->GetSequencerWidget()
 		]
 
 		+ SWidgetSwitcher::Slot()
 		[
-			Sequencer->GetSequencerWidget()
+			MakeHierarchyTabContent()
 		];
 
 	return SNew(SVerticalBox)
@@ -712,66 +712,58 @@ TSharedRef<SWidget> SCineAssemblySchemaWindow::MakeHierarchyTabContent()
 		.HAlign(HAlign_Left)
 		.Padding(0.0f, 0.0f, 0.0f, 8.0f)
 		[
-			SNew(SButton)
-				.ContentPadding(FMargin(2.0f))
-				.OnClicked(this, &SCineAssemblySchemaWindow::OnAddFolderClicked)
-				[
-					SNew(SHorizontalBox)
-
-					+ SHorizontalBox::Slot()
-						.AutoWidth()
-						.Padding(0.0f, 0.0f, 4.0f, 0.0f)
-						[
-							SNew(SImage)
-								.Image(FAppStyle::Get().GetBrush("Icons.Plus"))
-								.ColorAndOpacity(FStyleColors::AccentGreen)
-						]
+			SNew(SHorizontalBox)
 
-					+ SHorizontalBox::Slot()
-						.AutoWidth()
+			+ SHorizontalBox::Slot()
+				[
+					SNew(SButton)
+						.ContentPadding(FMargin(2.0f))
+						.OnClicked(this, &SCineAssemblySchemaWindow::OnAddAssetClicked)
 						[
-							SNew(STextBlock).Text(LOCTEXT("AddFolderButton", "Add Folder"))
+							SNew(SHorizontalBox)
+
+								+ SHorizontalBox::Slot()
+								.AutoWidth()
+								.Padding(0.0f, 0.0f, 4.0f, 0.0f)
+								[
+									SNew(SImage)
+										.Image(FAppStyle::Get().GetBrush("Icons.Plus"))
+										.ColorAndOpacity(FStyleColors::AccentGreen)
+								]
+
+								+ SHorizontalBox::Slot()
+								.AutoWidth()
+								[
+									SNew(STextBlock).Text(LOCTEXT("AddAssetButton", "Add Asset"))
+								]
 						]
-		]		]
-
-		+ SVerticalBox::Slot()
-		.FillHeight(1.0f)
-		[
-			SNew(SBorder)
-				.BorderImage(FCineAssemblyToolsStyle::Get().GetBrush("ProductionWizard.RecessedBackground"))
-				[
-					TreeView.ToSharedRef()
 				]
-		];
 
-	TSharedRef<SWidget> AssetListWidget = SNew(SVerticalBox)
-
-		+ SVerticalBox::Slot()
-		.AutoHeight()
-		.HAlign(HAlign_Left)
-		.Padding(0.0f, 0.0f, 0.0f, 8.0f)
-		[
-			SNew(SButton)
-				.ContentPadding(FMargin(2.0f))
-				.OnClicked(this, &SCineAssemblySchemaWindow::OnAddAssetClicked)
+				+ SHorizontalBox::Slot()
 				[
-					SNew(SHorizontalBox)
-
-					+ SHorizontalBox::Slot()
-						.AutoWidth()
-						.Padding(0.0f, 0.0f, 4.0f, 0.0f)
+					SNew(SButton)
+						.ContentPadding(FMargin(2.0f))
+						.OnClicked(this, &SCineAssemblySchemaWindow::OnAddFolderClicked)
 						[
-							SNew(SImage)
-								.Image(FAppStyle::Get().GetBrush("Icons.Plus"))
-								.ColorAndOpacity(FStyleColors::AccentGreen)
+							SNew(SHorizontalBox)
+
+								+ SHorizontalBox::Slot()
+								.AutoWidth()
+								.Padding(0.0f, 0.0f, 4.0f, 0.0f)
+								[
+									SNew(SImage)
+										.Image(FAppStyle::Get().GetBrush("Icons.Plus"))
+										.ColorAndOpacity(FStyleColors::AccentGreen)
+								]
+
+								+ SHorizontalBox::Slot()
+								.AutoWidth()
+								[
+									SNew(STextBlock).Text(LOCTEXT("AddFolderButton", "Add Folder"))
+								]
 						]
-
-					+ SHorizontalBox::Slot()
-						.AutoWidth()
-						[
-							SNew(STextBlock).Text(LOCTEXT("AddAssetButton", "Add Asset"))
-						]
-		]		]
+				]
+		]
 
 		+ SVerticalBox::Slot()
 		.FillHeight(1.0f)
@@ -779,14 +771,7 @@ TSharedRef<SWidget> SCineAssemblySchemaWindow::MakeHierarchyTabContent()
 			SNew(SBorder)
 				.BorderImage(FCineAssemblyToolsStyle::Get().GetBrush("ProductionWizard.RecessedBackground"))
 				[
-					SAssignNew(AssetListView, SListView<TSharedPtr<FSchemaAssetParameters>>)
-						.ListItemsSource(&AssetListItems)
-						.SelectionMode(ESelectionMode::Single)
-						.OnGenerateRow(this, &SCineAssemblySchemaWindow::OnGenerateAssetRow)
-						.OnItemsRebuilt(this, &SCineAssemblySchemaWindow::OnAssetListRebuilt)
-						.OnContextMenuOpening(this, &SCineAssemblySchemaWindow::MakeAssetListContextMenu)
-						.OnMouseButtonDoubleClick(this, &SCineAssemblySchemaWindow::EnterEditMode)
-						.OnKeyDownHandler(this, &SCineAssemblySchemaWindow::OnAssetListKeyDown)
+					TreeView.ToSharedRef()
 				]
 		];
 
@@ -811,36 +796,13 @@ TSharedRef<SWidget> SCineAssemblySchemaWindow::MakeHierarchyTabContent()
 				[
 					SNew(STextBlock)
 						.AutoWrapText(true)
-						.Text(LOCTEXT("AssetListInstructions", "Add new named subsequences to the list on the left, then drag and drop them to a location in the folder tree on the right. "
-							"When a new Cine Assembly is created using this Schema, the subsequences will automatically be created and added as tracks to the Assembly."))
+						.Text(LOCTEXT("AssetListInstructions", "The assets and folders displayed in this content tree will be created each time a new Cine Assembly is created using this Schema."))
 				]
 
 			+ SVerticalBox::Slot()
 				.FillHeight(1.0f)
 				[
-					SNew(SHorizontalBox)
-
-					+ SHorizontalBox::Slot()
-						.FillWidth(0.5f)
-						.Padding(0.0f, 0.0f, 16.0f, 0.0f)
-						[
-							AssetListWidget
-						]
-
-					+ SHorizontalBox::Slot()
-						.AutoWidth()
-						[
-							SNew(SSeparator)
-								.Orientation(Orient_Vertical)
-								.Thickness(2.0f)
-						]
-
-					+ SHorizontalBox::Slot()
-						.FillWidth(0.5f)
-						.Padding(16.0f, 0.0f, 0.0f, 0.0f)
-						[
-							ContentHierarchyWidget
-						]
+					ContentHierarchyWidget
 				]
 		];
 }
@@ -919,153 +881,6 @@ FReply SCineAssemblySchemaWindow::OnAddFolderClicked()
 	return FReply::Handled();
 }
 
-TSharedRef<ITableRow> SCineAssemblySchemaWindow::OnGenerateAssetRow(TSharedPtr<FSchemaAssetParameters> InItem, const TSharedRef<STableViewBase>& OwnerTable)
-{
-	TSharedRef<SSchemaAssetTableRow> Row = SNew(SSchemaAssetTableRow, OwnerTable)
-		.ShowSelection(true)
-		.Padding(FMargin(4.0f, 2.0f))
-		.OnDragDetected(this, &SCineAssemblySchemaWindow::OnAssetRowDragDetected)
-		.ToolTipText(LOCTEXT("AssetRowTooltip", "Set the name and template subsequence (optional) for this subsequence, then add it to the content tree on the right."));
-
-	// Store a reference to the editable textblock in the row to easily set it to edit mode for renaming
-	Row->TextBlock = SNew(SNamingTokensEditableTextBox)
-		.SupportInlineEdit(true)
-		.ShouldEvaluateTokens(false)
-		.Text_Lambda([InItem]()
-		{ 
-			return FText::FromString(InItem->Name); 
-		})
-		.OnTextCommitted_Lambda([this, InItem](const FText& InText, ETextCommit::Type InCommitType)
-		{
-			InItem->Name = InText.ToString();
-		})
-		.OnValidateTokenizedText(this, &SCineAssemblySchemaWindow::ValidateAssetName);
-
-	Row->SetContent(SNew(SBorder)
-		.Padding(8.0f)
-		.BorderImage(FCineAssemblyToolsStyle::Get().GetBrush("ProductionWizard.PanelBackground"))
-		[
-			SNew(SHorizontalBox)
-
-			+ SHorizontalBox::Slot()
-				.AutoWidth()
-				.VAlign(VAlign_Center)
-				.Padding(0.0f, 0.0f, 8.0f, 0.0f)
-				[
-					SNew(SImage).Image(FCineAssemblyToolsStyle::Get().GetBrush("Icons.Sequencer"))
-				]
-
-			+ SHorizontalBox::Slot()
-				.FillWidth(1.0f)
-				.VAlign(VAlign_Center)
-				[
-					Row->TextBlock.ToSharedRef()
-				]
-
-			+ SHorizontalBox::Slot()
-				.FillWidth(1.0f)
-				.VAlign(VAlign_Center)
-				[
-					SNew(SObjectPropertyEntryBox)
-						.AllowCreate(false)
-						.AllowClear(true)
-						.DisplayThumbnail(false)
-						.DisplayCompactSize(true)
-						.AllowedClass(ULevelSequence::StaticClass())
-						.ToolTipText(LOCTEXT("SchemaAssetTemplateSelectionTooltip", "(Optional) Select an existing level sequence to use as a template for this subsequence."))
-						.ObjectPath_Lambda([WeakInItem = TWeakPtr<FSchemaAssetParameters>(InItem)]() -> FString
-						{
-							if (TSharedPtr<FSchemaAssetParameters> ItemPtr = WeakInItem.Pin())
-							{
-								return ItemPtr->Template.ToString();
-							}
-							return FString();
-						})
-						.OnObjectChanged_Lambda([WeakInItem = TWeakPtr<FSchemaAssetParameters>(InItem)](const FAssetData& InAssetData)
-						{
-							if (TSharedPtr<FSchemaAssetParameters> ItemPtr = WeakInItem.Pin())
-							{
-								ItemPtr->Template = InAssetData.GetSoftObjectPath();
-							}
-						})
-				]
-
-			+ SHorizontalBox::Slot()
-				.AutoWidth()
-				.VAlign(VAlign_Center)
-				[
-					SNew(SImage).Image(FCineAssemblyToolsStyle::Get().GetBrush("Icons.DragHandle"))
-				]
-		]);
-
-	return Row;
-}
-
-void SCineAssemblySchemaWindow::OnAssetListRebuilt()
-{
-	if (AssetListView && AssetListView->GetItems().Num() > 0)
-	{
-		// Only trigger edit mode on the last item in the asset list, and only if it is an empty string so that it can be renamed
-		TSharedPtr<FSchemaAssetParameters> LastItem = AssetListView->GetItems().Last();
-		if (LastItem && LastItem->Name.IsEmpty())
-		{
-			LastItem->Name = TEXT("NewAsset");
-			EnterEditMode(LastItem);
-		}
-	}
-}
-
-TSharedPtr<SWidget> SCineAssemblySchemaWindow::MakeAssetListContextMenu()
-{
-	// The ListView uses single selection mode, so at most one item can ever be selected by the user
-	TArray<TSharedPtr<FSchemaAssetParameters>> SelectedItems = AssetListView->GetSelectedItems();
-	if (SelectedItems.Num() == 1)
-	{
-		constexpr bool bCloseAfterMenuSelection = true;
-		FMenuBuilder MenuBuilder(bCloseAfterMenuSelection, nullptr);
-
-		TSharedPtr<FSchemaAssetParameters>& SelectedItem = SelectedItems[0];
-
-		MenuBuilder.AddMenuEntry(
-			LOCTEXT("RenameAsset", "Rename"),
-			LOCTEXT("RenameAssetToolTip", "Rename"),
-			FSlateIcon(FCineAssemblyToolsStyle::StyleName, "Icons.AssetNaming"),
-			FUIAction(FExecuteAction::CreateSP(this, &SCineAssemblySchemaWindow::EnterEditMode, SelectedItem)),
-			NAME_None,
-			EUserInterfaceActionType::Button
-		);
-
-		MenuBuilder.AddMenuEntry(
-			LOCTEXT("DeleteAssetFromList", "Delete"),
-			LOCTEXT("DeleteAssetFromListToolTip", "Delete"),
-			FSlateIcon(FAppStyle::GetAppStyleSetName(), "Icons.Delete"),
-			FUIAction(FExecuteAction::CreateSP(this, &SCineAssemblySchemaWindow::DeleteAssetItem, SelectedItem)),
-			NAME_None,
-			EUserInterfaceActionType::Button
-		);
-
-		return MenuBuilder.MakeWidget();
-	}
-
-	return SNullWidget::NullWidget;
-}
-
-FReply SCineAssemblySchemaWindow::OnAssetListKeyDown(const FGeometry& MyGeometry, const FKeyEvent& InKeyEvent)
-{
-	if (InKeyEvent.GetKey() == EKeys::Delete)
-	{
-		// The list view uses single selection mode, so at most one item can ever be selected by the user
-		TArray<TSharedPtr<FSchemaAssetParameters>> SelectedItems = AssetListView->GetSelectedItems();
-		if (SelectedItems.Num() == 1)
-		{
-			const TSharedPtr<FSchemaAssetParameters>& SelectedItem = SelectedItems[0];
-			DeleteAssetItem(SelectedItem);
-		}
-	}
-
-	return FReply::Handled();
-}
-
 TSharedRef<ITableRow> SCineAssemblySchemaWindow::OnGenerateTreeRow(TSharedPtr<ISchemaContentHierarchyTreeItem> TreeItem, const TSharedRef<STableViewBase>& OwnerTable)
 {
 	TSharedPtr<SImage> Icon = SNew(SImage);
@@ -1367,18 +1182,6 @@ void SCineAssemblySchemaWindow::SetChildrenPathRecursive(TSharedPtr<FSchemaConte
 	}
 }
 
-void SCineAssemblySchemaWindow::EnterEditMode(TSharedPtr<FSchemaAssetParameters> ItemToRename)
-{
-	// Get the row for the input item and put its textblock in edit mode so the user can rename the item
-	if (TSharedPtr<SSchemaAssetTableRow> Widget = StaticCastSharedPtr<SSchemaAssetTableRow>(AssetListView->WidgetFromItem(ItemToRename)))
-	{
-		if (Widget->TextBlock)
-		{
-			Widget->TextBlock->EnterEditingMode();
-		}
-	}
-}
-
 void SCineAssemblySchemaWindow::OnTreeItemTextCommitted(const FText& InText, ETextCommit::Type InCommitType, TSharedPtr<ISchemaContentHierarchyTreeItem> TreeItem)
 {
 	// Early-out if the name has not actually changed
@@ -1475,16 +1278,6 @@ bool SCineAssemblySchemaWindow::ValidateFolderName(const FText& InText, FText& O
 	return true;
 }
 
-void SCineAssemblySchemaWindow::DeleteAssetItem(TSharedPtr<FSchemaAssetParameters> InItem)
-{
-	const int32 IndexToRemove = AssetListItems.IndexOfByKey(InItem);
-	if (AssetListItems.IsValidIndex(IndexToRemove))
-	{
-		AssetListItems.RemoveAt(IndexToRemove);
-	}
-	AssetListView->RequestListRefresh();
-}
-
 void SCineAssemblySchemaWindow::DeleteTreeItem(TSharedPtr<ISchemaContentHierarchyTreeItem> TreeItem)
 {
 	if (TSharedPtr<FSchemaContentHierarchyFolder> FolderItem = TreeItem->AsFolder())
@@ -1599,16 +1392,6 @@ void SCineAssemblySchemaWindow::OnSchemaPropertiesChanged(const FPropertyChanged
 
 FReply SCineAssemblySchemaWindow::OnCreateAssetClicked()
 {
-	// Prevent the user from finishing creating the schema if there are still assets in the list view that have not been placed in the content hierarchy
-	if (AssetListItems.Num() > 0)
-	{
-		const FText DialogMessage = LOCTEXT("RemainingAssetDialog", "The Content Hierarchy tab contains named subsequences that have not yet been placed in the folder tree. "
-			"Please drag and drop the remaining subsequences, or delete them from the asset list if they are not needed.");
-		FMessageDialog::Open(EAppMsgType::Ok, DialogMessage);
-
-		return FReply::Unhandled();
-	}
-
 	UCineAssemblySchemaFactory::CreateConfiguredSchema(SchemaToConfigure.Get(), CreateAssetPath);
 
 	TSharedPtr<SWindow> ContainingWindow = FSlateApplication::Get().FindWidgetWindow(AsShared());
@@ -1816,40 +1599,6 @@ EVisibility SCineAssemblySchemaWindow::GetDisplayRateWarningVisibility() const
 	return EVisibility::Collapsed;
 }
 
-FReply SCineAssemblySchemaWindow::OnAssetRowDragDetected(const FGeometry& MyGeometry, const FPointerEvent& MouseEvent)
-{
-	if (MouseEvent.IsMouseButtonDown(EKeys::LeftMouseButton))
-	{
-		// The ListView uses single selection mode, so at most one item can ever be selected by the user
-		TArray<TSharedPtr<FSchemaAssetParameters>> SelectedItems = AssetListView->GetSelectedItems();
-		if (SelectedItems.Num() == 1)
-		{
-			const TSharedPtr<FSchemaAssetParameters>& SelectedItem = SelectedItems[0];
-			TSharedRef<FSchemaAssetDragDrop> Operation = FSchemaAssetDragDrop::New(*SelectedItem);
-
-			Operation->OnDropNotHandled.BindLambda([this, WeakSelectedItem = TWeakPtr<FSchemaAssetParameters>(SelectedItem)](const FString& AssetName)
-				{
-					FSoftObjectPath RevertTemplate;
-					if (TSharedPtr<FSchemaAssetParameters> SelectedItemPtr = WeakSelectedItem.Pin())
-					{
-						RevertTemplate = SelectedItemPtr->Template;
-					}
-					// If the asset item is not dropped in a valid place, restore it to the asset list
-					AssetListItems.Add(MakeShared<FSchemaAssetParameters>(AssetName, RevertTemplate));
-					AssetListView->RequestListRefresh();
-				});
-
-			// Remove this item from the asset list while it is being dragged
-			AssetListItems.Remove(SelectedItems[0]);
-			AssetListView->RequestListRefresh();
-
-			return FReply::Handled().BeginDragDrop(Operation);
-		}
-	}
-
-	return FReply::Unhandled();
-}
-
 FReply SCineAssemblySchemaWindow::OnTreeRowDragDetected(const FGeometry& MyGeometry, const FPointerEvent& MouseEvent)
 {
 	if (MouseEvent.IsMouseButtonDown(EKeys::LeftMouseButton))
diff --git a/Engine/Plugins/VirtualProduction/CinematicAssemblyTools/Source/CineAssemblyToolsEditor/Private/UI/CineAssembly/SCineAssemblySchemaWindow.h b/Engine/Plugins/VirtualProduction/CinematicAssemblyTools/Source/CineAssemblyToolsEditor/Private/UI/CineAssembly/SCineAssemblySchemaWindow.h
--- a/Engine/Plugins/VirtualProduction/CinematicAssemblyTools/Source/CineAssemblyToolsEditor/Private/UI/CineAssembly/SCineAssemblySchemaWindow.h
+++ b/Engine/Plugins/VirtualProduction/CinematicAssemblyTools/Source/CineAssemblyToolsEditor/Private/UI/CineAssembly/SCineAssemblySchemaWindow.h
@@ -9,7 +9,6 @@
 #include "ISequencer.h"
 #include "SNamingTokensEditableTextBox.h"
 #include "Widgets/Layout/SWidgetSwitcher.h"
-#include "Widgets/Views/SListView.h"
 #include "Widgets/Views/STreeView.h"
 
 struct FPropertyAndParent;
@@ -35,14 +34,6 @@ struct FSchemaAssetParameters
 	FSoftObjectPath Template;
 };
 
-/** Row widget for the asset list view */
-class SSchemaAssetTableRow : public STableRow<TSharedPtr<FSchemaAssetParameters>>
-{
-public:
-	/** Editable text widget stored in this row to easily trigger edit mode */
-	TSharedPtr<SNamingTokensEditableTextBox> TextBlock;
-};
-
 /** Drag and Drop operation to handle dragging schema assets and folders  */
 class FSchemaAssetDragDrop : public FDragDropOperation
 {
@@ -62,7 +53,7 @@ public:
 	virtual TSharedPtr<SWidget> GetDefaultDecorator() const override;
 	// End FDragDropOperation overrides
 
-	/** The source item being dragged. This could be null if the thing being dragged is from the Asset List (not the tree) */
+	/** The source item being dragged. */
 	TSharedPtr<ISchemaContentHierarchyTreeItem> SourceTreeItem;
 
 	/** The name of the item being displayed, used by the decorator and to name the new tree item that will get made when dropped */
@@ -155,18 +146,6 @@ private:
 	/** Adds a new folder item to the content tree view */
 	FReply OnAddFolderClicked();
 
-	/** Generates a row in the asset list view */
-	TSharedRef<ITableRow> OnGenerateAssetRow(TSharedPtr<FSchemaAssetParameters> InItem, const TSharedRef<STableViewBase>& OwnerTable);
-
-	/** Callback when the asset list view rebuilds itself, used to make the textbox in the most recently added row editable */
-	void OnAssetListRebuilt();
-
-	/** Create the context menu when the asset list view is right-clicked */
-	TSharedPtr<SWidget> MakeAssetListContextMenu();
-
-	/** Handles key presses on the asset list */
-	FReply OnAssetListKeyDown(const FGeometry& MyGeometry, const FKeyEvent& InKeyEvent);
-
 	/** Generates the row widget for an entry in the tree view */
 	TSharedRef<ITableRow> OnGenerateTreeRow(TSharedPtr<ISchemaContentHierarchyTreeItem> TreeItem, const TSharedRef<STableViewBase>& OwnerTable);
 
@@ -209,9 +188,6 @@ private:
 	/** Modifies the path of the input folder item and all of its children (recursively) with the input new path */
 	void SetChildrenPathRecursive(TSharedPtr<FSchemaContentHierarchyFolder> FolderItem, const FString& NewPath);
 
-	/** Finds the row for the input item and puts its textblock into edit mode */
-	void EnterEditMode(TSharedPtr<FSchemaAssetParameters> ItemToRename);
-
 	/** Renames the tree item, and updates the paths of all of its children */
 	void OnTreeItemTextCommitted(const FText& InText, ETextCommit::Type InCommitType, TSharedPtr<ISchemaContentHierarchyTreeItem> TreeItem);
 
@@ -227,9 +203,6 @@ private:
 	/** Validate the text entered by the user to ensure it will be a valid folder name */
 	bool ValidateFolderName(const FText& InText, FText& OutErrorMessage, TSharedPtr<FSchemaContentHierarchyFolder> FolderItem) const;
 
-	/** Remove the input item from the asset list */
-	void DeleteAssetItem(TSharedPtr<FSchemaAssetParameters> InItem);
-
 	/** Remove the input item and its children from the tree view */
 	void DeleteTreeItem(TSharedPtr<ISchemaContentHierarchyTreeItem> TreeItem);
 
@@ -239,9 +212,6 @@ private:
 	/** Removes a folder and all of its children from the tree view */
 	void DeleteFolderItem(TSharedPtr<FSchemaContentHierarchyFolder> FolderItem);
 
-	/** Begins a drag and drop event to drag an item out of the asset list view */
-	FReply OnAssetRowDragDetected(const FGeometry& MyGeometry, const FPointerEvent& MouseEvent);
-
 	/** Begins a drag and drop event to drag an item out of the content tree view */
 	FReply OnTreeRowDragDetected(const FGeometry& MyGeometry, const FPointerEvent& MouseEvent);
 
@@ -283,9 +253,6 @@ private:
 	/** Switcher that controls which menu tab is currently visible */
 	TSharedPtr<SWidgetSwitcher> MenuTabSwitcher;
 
-	/** Switcher that controls which tab widget in the Asset Tab is currently visible */
-	TSharedPtr<SWidgetSwitcher> AssetTabSwitcher;
-
 	/** Transient object used only by this UI to configure the properties of the new asset that will get created by the factory */
 	TStrongObjectPtr<UCineAssemblySchema> SchemaToConfigure = nullptr;
 
@@ -295,12 +262,6 @@ private:
 	/** The root path where the configured schema will be created */
 	FString CreateAssetPath;
 
-	/** List items sources for the asset list view */
-	TArray<TSharedPtr<FSchemaAssetParameters>> AssetListItems;
-
-	/** List view of assets that will be associate with the schema */
-	TSharedPtr<SListView<TSharedPtr<FSchemaAssetParameters>>> AssetListView;
-
 	/** Items source for the tree view */
 	TArray<TSharedPtr<ISchemaContentHierarchyTreeItem>> TreeItems;
 
