ShotManagement: Refactor to deprecate two properties of CineAssemblySchema, DefaultAssemblyPath and FoldersToCreate, and use equivalent CineAssembly properties instead. Tools that used to read/write those properties from the schema now read/write from the properties of the Schema's template assembly instead.
#rnx
#virtualized
#Synced-CL 51715153

diff --git a/Engine/Plugins/VirtualProduction/CinematicAssemblyTools/Source/CineAssemblyTools/Private/CineAssembly.cpp b/Engine/Plugins/VirtualProduction/CinematicAssemblyTools/Source/CineAssemblyTools/Private/CineAssembly.cpp
--- a/Engine/Plugins/VirtualProduction/CinematicAssemblyTools/Source/CineAssemblyTools/Private/CineAssembly.cpp
+++ b/Engine/Plugins/VirtualProduction/CinematicAssemblyTools/Source/CineAssemblyTools/Private/CineAssembly.cpp
@@ -165,7 +165,7 @@ void UCineAssembly::Initialize()
 
 void UCineAssembly::InitializeFromSchema(UCineAssemblySchema* InSchema)
 {
-	SetSchema(InSchema);
+	ChangeSchema(InSchema);
 
 	if (InSchema)
 	{
@@ -186,6 +186,12 @@ void UCineAssembly::InitializeFromTemplate(ULevelSequence* Template)
 		return;
 	}
 
+	if (UCineAssembly* TemplateAssembly = Cast<UCineAssembly>(Template))
+	{
+		PathRelativeToRoot = TemplateAssembly->PathRelativeToRoot;
+		DefaultFolderNames = TemplateAssembly->DefaultFolderNames;
+	}
+
 	if (UMovieScene* TemplateMovieScene = Template->GetMovieScene())
 	{
 		// Attempt to duplicate the inner MovieScene. Assign only if duplication is a success to avoid leaving the assembly in a partially valid state.
@@ -304,16 +310,11 @@ void UCineAssembly::ChangeSchema(UCineAssemblySchema* InSchema)
 	}
 
 	// Reset the list of folders names to create from the Schema
+	PathRelativeToRoot.Template.Reset();
+	PathRelativeToRoot.Resolved = FText::GetEmpty();
 	DefaultFolderNames.Reset();
-	if (BaseSchema)
-	{
-		for (const FString& FolderName : BaseSchema->FoldersToCreate)
-		{
-			FTemplateString NewTemplateName;
-			NewTemplateName.Template = FolderName;
-			DefaultFolderNames.Add(NewTemplateName);
-		}
-	}
+
+	SubAssemblies.Empty();
 
 	// Update the Data Only flag on the assembly
 	if (BaseSchema)
@@ -393,11 +394,11 @@ void UCineAssembly::ConvertSubAssemblyTracks()
 
 					// Set the name and path of the new SubAssembly from the data stored in the SubAssembly section
 					NewSubAssembly->AssemblyName.Template = SubAssemblySection->GetSequenceName().ToString();
-					NewSubAssembly->RelativeAssemblyPath.Template = SubAssemblySection->GetSequencePath();
+					NewSubAssembly->PathRelativeToParent.Template = SubAssemblySection->GetSequencePath();
 
-					if (!Algo::ContainsBy(DefaultFolderNames, NewSubAssembly->RelativeAssemblyPath.Template, &FTemplateString::Template))
+					if (!Algo::ContainsBy(DefaultFolderNames, NewSubAssembly->PathRelativeToParent.Template, &FTemplateString::Template))
 					{
-						DefaultFolderNames.Add(NewSubAssembly->RelativeAssemblyPath);
+						DefaultFolderNames.Add(NewSubAssembly->PathRelativeToParent);
 					}
 
 					NewSubSection->SetSequence(NewSubAssembly);
@@ -450,7 +451,12 @@ void UCineAssembly::DuplicateSubSequences(ULevelSequence* SequenceTemplate)
 					}
 
 					// When the transient SubAssembly is converted into an asset, the relative path will tell the factory where to create the package relative to the root folder
-					DuplicateSubAssembly->RelativeAssemblyPath.Template = GetRelativePath(SubSequence, SequenceTemplate);
+					DuplicateSubAssembly->PathRelativeToParent.Template = GetRelativePath(SubSequence, SequenceTemplate);
+
+					if (!Algo::ContainsBy(DefaultFolderNames, DuplicateSubAssembly->PathRelativeToParent.Template, &FTemplateString::Template))
+					{
+						DefaultFolderNames.Add(DuplicateSubAssembly->PathRelativeToParent);
+					}
 
 					// Replace the sequence in the section with the duplicate assembly
 					SubSection->SetSequence(DuplicateSubAssembly);
@@ -957,14 +963,8 @@ void UCineAssembly::GetAssetPathAndRootFolder(FString& OutAssetPath, FString& Ou
 	const FAssetData AssemblyAssetData = FAssetData(this);
 	OutAssetPath = AssemblyAssetData.PackagePath.ToString();
 
-	OutRootFolder = OutAssetPath;
-	if (BaseSchema)
-	{
-		const FString DefaultAssemblyPath = BaseSchema->DefaultAssemblyPath;
-		const FString ResolvedAssemblyPath = UCineAssemblyNamingTokens::GetResolvedText(DefaultAssemblyPath, this).ToString();
-
-		OutRootFolder = OutRootFolder.Replace(*ResolvedAssemblyPath, TEXT(""));
-	}
+	const FString ResolvedAssemblyPath = UCineAssemblyNamingTokens::GetResolvedText(PathRelativeToRoot.Template, this).ToString();
+	OutRootFolder = OutAssetPath.Replace(*ResolvedAssemblyPath, TEXT(""));
 }
 
 #undef LOCTEXT_NAMESPACE
diff --git a/Engine/Plugins/VirtualProduction/CinematicAssemblyTools/Source/CineAssemblyTools/Private/CineAssemblySchema.cpp b/Engine/Plugins/VirtualProduction/CinematicAssemblyTools/Source/CineAssemblyTools/Private/CineAssemblySchema.cpp
--- a/Engine/Plugins/VirtualProduction/CinematicAssemblyTools/Source/CineAssemblyTools/Private/CineAssemblySchema.cpp
+++ b/Engine/Plugins/VirtualProduction/CinematicAssemblyTools/Source/CineAssemblyTools/Private/CineAssemblySchema.cpp
@@ -129,6 +129,21 @@ void UCineAssemblySchema::PostLoad()
 		}
 
 		SubsequencesToCreate_DEPRECATED.Empty();
+
+		// Store the old deprecated default assembly path in the template sequence
+		if (!DefaultAssemblyPath_DEPRECATED.IsEmpty())
+		{
+			TemplateSequence->PathRelativeToRoot.Template = DefaultAssemblyPath_DEPRECATED;
+			DefaultAssemblyPath_DEPRECATED.Empty();
+		}
+
+		// Store the old deprecated list of folders to create in the template sequence
+		if (!FoldersToCreate_DEPRECATED.IsEmpty())
+		{
+			TemplateSequence->DefaultFolderNames.Reset(FoldersToCreate_DEPRECATED.Num());
+			Algo::Transform(FoldersToCreate_DEPRECATED, TemplateSequence->DefaultFolderNames, [](const FString& TemplateString) { return FTemplateString(TemplateString, FText::GetEmpty()); });
+			FoldersToCreate_DEPRECATED.Empty();
+		}
 	}
 	PRAGMA_ENABLE_DEPRECATION_WARNINGS
 #endif // WITH_EDITOR
diff --git a/Engine/Plugins/VirtualProduction/CinematicAssemblyTools/Source/CineAssemblyTools/Public/CineAssembly.h b/Engine/Plugins/VirtualProduction/CinematicAssemblyTools/Source/CineAssemblyTools/Public/CineAssembly.h
--- a/Engine/Plugins/VirtualProduction/CinematicAssemblyTools/Source/CineAssemblyTools/Public/CineAssembly.h
+++ b/Engine/Plugins/VirtualProduction/CinematicAssemblyTools/Source/CineAssemblyTools/Public/CineAssembly.h
@@ -230,7 +230,11 @@ public:
 
 	/** The path where this assembly should be created, relative to its content root */
 	UPROPERTY()
-	FTemplateString RelativeAssemblyPath;
+	FTemplateString PathRelativeToRoot;
+
+	/** The path where this assembly should be created, relative to its parent assembly */
+	UPROPERTY()
+	FTemplateString PathRelativeToParent;
 
 	/** The level to open before opening this asset in Sequencer */
 	UPROPERTY(EditAnywhere, Category = "Default", meta = (AllowedClasses = "/Script/Engine.World"))
diff --git a/Engine/Plugins/VirtualProduction/CinematicAssemblyTools/Source/CineAssemblyTools/Public/CineAssemblySchema.h b/Engine/Plugins/VirtualProduction/CinematicAssemblyTools/Source/CineAssemblyTools/Public/CineAssemblySchema.h
--- a/Engine/Plugins/VirtualProduction/CinematicAssemblyTools/Source/CineAssemblyTools/Public/CineAssemblySchema.h
+++ b/Engine/Plugins/VirtualProduction/CinematicAssemblyTools/Source/CineAssemblyTools/Public/CineAssemblySchema.h
@@ -97,13 +97,6 @@ public:
 	UPROPERTY(EditAnywhere, Category = "Default")
 	FString DefaultAssemblyName;
 
-	/** 
-	 * The default path to use when creating assemblies from this schema
-	 * When an assembly asset is created, this path will be appended to path where the asset would normally have been created.
-	 */
-	UPROPERTY()
-	FString DefaultAssemblyPath;
-
 	/** Restricts assemblies made from this schema to using this Schema when picking a Parent Assembly */
 	UPROPERTY(EditAnywhere, Category = "Default", meta = (AllowedClasses = "/Script/CineAssemblyTools.CineAssemblySchema"))
 	FSoftObjectPath ParentSchema;
@@ -116,10 +109,6 @@ public:
 	UPROPERTY(EditAnywhere, Category = "Metadata")
 	TArray<FAssemblyMetadataDesc> AssemblyMetadata;
 
-	/** Paths of folders that should be created for assemblies that use this schema, relative to the path of the top-level assembly */
-	UPROPERTY()
-	TArray<FString> FoldersToCreate;
-
 	/** Template sequence, owned and managed by the schema, to use when creating the assembly as a starting point.  */
 	UPROPERTY()
 	TObjectPtr<UCineAssembly> TemplateSequence;
@@ -147,6 +136,14 @@ protected:
 	UE_DEPRECATED(5.8, "Subsequence Template objects are now stored in the SubAssembly sections of the TemplateSequence.")
 	UPROPERTY()
 	TMap<FString, FSoftObjectPath> SubsequenceTemplates_DEPRECATED;
+
+	UE_DEPRECATED(5.8, "This property is now managed directly by the TemplateSequence's PathRelativeToRoot property")
+	UPROPERTY()
+	FString DefaultAssemblyPath_DEPRECATED;
+
+	UE_DEPRECATED(5.8, "This property is now managed directly by the TemplateSequence's DefaultFolderNames property")
+	UPROPERTY()
+	TArray<FString> FoldersToCreate_DEPRECATED;
 #endif // WITH_EDITORONLY_DATA 
 
 private:
diff --git a/Engine/Plugins/VirtualProduction/CinematicAssemblyTools/Source/CineAssemblyToolsEditor/Private/CineAssemblyFactory.cpp b/Engine/Plugins/VirtualProduction/CinematicAssemblyTools/Source/CineAssemblyToolsEditor/Private/CineAssemblyFactory.cpp
--- a/Engine/Plugins/VirtualProduction/CinematicAssemblyTools/Source/CineAssemblyToolsEditor/Private/CineAssemblyFactory.cpp
+++ b/Engine/Plugins/VirtualProduction/CinematicAssemblyTools/Source/CineAssemblyToolsEditor/Private/CineAssemblyFactory.cpp
@@ -65,17 +65,10 @@ UObject* UCineAssemblyFactory::FactoryCreateNew(UClass* Class, UObject* InParent
 
 FString UCineAssemblyFactory::MakeAssemblyPackageName(UCineAssembly* ConfiguredAssembly, const FString& CreateAssetPath)
 {
-	// Resolve the default assembly path
-	FString DefaultAssemblyPath;
-	if (const UCineAssemblySchema* Schema = ConfiguredAssembly->GetSchema())
-	{
-		if (!Schema->DefaultAssemblyPath.IsEmpty())
-		{
-			DefaultAssemblyPath = UCineAssemblyNamingTokens::GetResolvedText(Schema->DefaultAssemblyPath, ConfiguredAssembly).ToString();
-		}
-	}
+	// Resolve the assembly's path relative to its content root
+	ConfiguredAssembly->PathRelativeToRoot.Resolved = UCineAssemblyNamingTokens::GetResolvedText(ConfiguredAssembly->PathRelativeToRoot.Template, ConfiguredAssembly);
 
-	return CreateAssetPath / DefaultAssemblyPath / ConfiguredAssembly->AssemblyName.Resolved.ToString();
+	return CreateAssetPath / ConfiguredAssembly->PathRelativeToRoot.Resolved.ToString() / ConfiguredAssembly->AssemblyName.Resolved.ToString();
 }
 
 bool UCineAssemblyFactory::MakeUniqueNameAndPath(UCineAssembly* ConfiguredAssembly, const FString& CreateAssetPath, FString& UniquePackageName, FString& UniqueAssetName)
@@ -201,9 +194,9 @@ void UCineAssemblyFactory::CreateConfiguredSubAssemblies(UCineAssembly* Configur
 						SubAssembly->Production = ConfiguredAssembly->Production;
 						SubAssembly->ProductionName = ConfiguredAssembly->ProductionName;
 
-						SubAssembly->RelativeAssemblyPath.Resolved = UCineAssemblyNamingTokens::GetResolvedText(SubAssembly->RelativeAssemblyPath.Template, ConfiguredAssembly);
+						SubAssembly->PathRelativeToParent.Resolved = UCineAssemblyNamingTokens::GetResolvedText(SubAssembly->PathRelativeToParent.Template, ConfiguredAssembly);
 
-						CreateConfiguredAssembly(SubAssembly, CreateAssetPath / SubAssembly->RelativeAssemblyPath.Resolved.ToString());
+						CreateConfiguredAssembly(SubAssembly, CreateAssetPath / SubAssembly->PathRelativeToParent.Resolved.ToString());
 					}
 				}
 			}
diff --git a/Engine/Plugins/VirtualProduction/CinematicAssemblyTools/Source/CineAssemblyToolsEditor/Private/TakeRecorder/CineAssemblyTakeRecorderIntegration.cpp b/Engine/Plugins/VirtualProduction/CinematicAssemblyTools/Source/CineAssemblyToolsEditor/Private/TakeRecorder/CineAssemblyTakeRecorderIntegration.cpp
--- a/Engine/Plugins/VirtualProduction/CinematicAssemblyTools/Source/CineAssemblyToolsEditor/Private/TakeRecorder/CineAssemblyTakeRecorderIntegration.cpp
+++ b/Engine/Plugins/VirtualProduction/CinematicAssemblyTools/Source/CineAssemblyToolsEditor/Private/TakeRecorder/CineAssemblyTakeRecorderIntegration.cpp
@@ -82,7 +82,7 @@ void FCineAssemblyTakeRecorderIntegration::OnRecordingInitialized(UTakeRecorder*
 				CurrentAssemblyPath = FPaths::GetPath(CurrentAssemblyPath);
 
 				// If the schema defines some additional default assembly path, move the assembly asset accordingly
-				if (!Schema->DefaultAssemblyPath.IsEmpty())
+				if (!Assembly->PathRelativeToRoot.Template.IsEmpty())
 				{
 					FString UniquePackageName;
 					FString UniqueAssetName;
diff --git a/Engine/Plugins/VirtualProduction/CinematicAssemblyTools/Source/CineAssemblyToolsEditor/Private/UI/CineAssembly/SCineAssemblyConfigPanel.cpp b/Engine/Plugins/VirtualProduction/CinematicAssemblyTools/Source/CineAssemblyToolsEditor/Private/UI/CineAssembly/SCineAssemblyConfigPanel.cpp
--- a/Engine/Plugins/VirtualProduction/CinematicAssemblyTools/Source/CineAssemblyToolsEditor/Private/UI/CineAssembly/SCineAssemblyConfigPanel.cpp
+++ b/Engine/Plugins/VirtualProduction/CinematicAssemblyTools/Source/CineAssemblyToolsEditor/Private/UI/CineAssembly/SCineAssemblyConfigPanel.cpp
@@ -343,20 +343,11 @@ void SCineAssemblyConfigPanel::PopulateHierarchyTree()
 
 	if (const UCineAssemblySchema* Schema = CineAssemblyToConfigure->GetSchema())
 	{
-		TArray<FTemplateString> FolderTemplates;
-		FolderTemplates.Reserve(Schema->FoldersToCreate.Num());
-		Algo::Transform(Schema->FoldersToCreate, FolderTemplates, [](const FString& TemplateString)
-			{
-				FTemplateString FolderTemplate;
-				FolderTemplate.Template = TemplateString;
-				return FolderTemplate;
-			});
-
-		AddItemsToTree(FolderTemplates, FHierarchyTreeItem::EItemType::Folder);
+		AddItemsToTree(CineAssemblyToConfigure->DefaultFolderNames, FHierarchyTreeItem::EItemType::Folder);
 		AddItemsToTree(CineAssemblyToConfigure->SubAssemblyNames, FHierarchyTreeItem::EItemType::Asset);
 
 		// Add the top-level assembly tree item
-		TSharedPtr<FHierarchyTreeItem> ParentItem = FindItemAtPathRecursive(RootItem, Schema->DefaultAssemblyPath);
+		TSharedPtr<FHierarchyTreeItem> ParentItem = FindItemAtPathRecursive(RootItem, CineAssemblyToConfigure->PathRelativeToRoot.Template);
 		if (ParentItem)
 		{
 			TSharedPtr<FHierarchyTreeItem> NewItem = MakeShared<FHierarchyTreeItem>();
diff --git a/Engine/Plugins/VirtualProduction/CinematicAssemblyTools/Source/CineAssemblyToolsEditor/Private/UI/CineAssembly/SCineAssemblySchemaWindow.cpp b/Engine/Plugins/VirtualProduction/CinematicAssemblyTools/Source/CineAssemblyToolsEditor/Private/UI/CineAssembly/SCineAssemblySchemaWindow.cpp
--- a/Engine/Plugins/VirtualProduction/CinematicAssemblyTools/Source/CineAssemblyToolsEditor/Private/UI/CineAssembly/SCineAssemblySchemaWindow.cpp
+++ b/Engine/Plugins/VirtualProduction/CinematicAssemblyTools/Source/CineAssemblyToolsEditor/Private/UI/CineAssembly/SCineAssemblySchemaWindow.cpp
@@ -1088,9 +1088,10 @@ void SCineAssemblySchemaWindow::InitializeContentTree()
 	TreeItems.Add(RootFolder);
 
 	// Sort the list of folders so that paths are added to the tree in the proper order
-	SchemaToConfigure->FoldersToCreate.Sort();
+	TArray<FString> FolderTemplates = GetFolderTemplates();
+	FolderTemplates.Sort();
 
-	for (const FString& FolderPath : SchemaToConfigure->FoldersToCreate)
+	for (const FString& FolderPath : FolderTemplates)
 	{
 		const FString ParentPath = FPaths::GetPath(FolderPath);
 
@@ -1109,7 +1110,8 @@ void SCineAssemblySchemaWindow::InitializeContentTree()
 	UpdateContentHierarchyAssets();
 
 	// Add the top-level assembly node
-	TSharedPtr<FSchemaContentHierarchyFolder> ParentFolder = FindFolderByPathRecursive(RootFolder, SchemaToConfigure->DefaultAssemblyPath);
+	const FString& ToplevelAssemblyPath = GetDefaultAssemblyPath();
+	TSharedPtr<FSchemaContentHierarchyFolder> ParentFolder = FindFolderByPathRecursive(RootFolder, ToplevelAssemblyPath);
 	if (ParentFolder)
 	{
 		TopLevelAssemblyItem = MakeShared<FSchemaContentHierarchyAsset>();
@@ -1119,6 +1121,57 @@ void SCineAssemblySchemaWindow::InitializeContentTree()
 	}
 }
 
+TArray<FString> SCineAssemblySchemaWindow::GetFolderTemplates()
+{
+	if (!SchemaToConfigure || !SchemaToConfigure->TemplateSequence)
+	{
+		return TArray<FString>();
+	}
+
+	TArrayView<FTemplateString> DefaultFolderNames = SchemaToConfigure->TemplateSequence->DefaultFolderNames;
+
+	// Extract the template folder names from the template sequence's default folder list
+	TArray<FString> TemplateFolderNames;
+	TemplateFolderNames.Reserve(DefaultFolderNames.Num());
+
+	Algo::Transform(DefaultFolderNames, TemplateFolderNames, [](const FTemplateString& TemplateString) { return TemplateString.Template; });
+
+	return TemplateFolderNames;
+}
+
+void SCineAssemblySchemaWindow::SetFolderTemplates(TArray<FString> InFolderTemplates)
+{
+	if (!SchemaToConfigure || !SchemaToConfigure->TemplateSequence)
+	{
+		return;
+	}
+
+	TArray<FTemplateString>& DefaultFolderNames = SchemaToConfigure->TemplateSequence->DefaultFolderNames;
+
+	DefaultFolderNames.Reset(InFolderTemplates.Num());
+	Algo::Transform(InFolderTemplates, DefaultFolderNames, [](const FString& TemplateString) { return FTemplateString(TemplateString, FText::GetEmpty()); });
+}
+
+FString SCineAssemblySchemaWindow::GetDefaultAssemblyPath()
+{
+	if (!SchemaToConfigure || !SchemaToConfigure->TemplateSequence)
+	{
+		return FString();
+	}
+
+	return SchemaToConfigure->TemplateSequence->PathRelativeToRoot.Template;
+}
+
+void SCineAssemblySchemaWindow::SetDefaultAssemblyPath(const FString& InPath)
+{
+	if (!SchemaToConfigure || !SchemaToConfigure->TemplateSequence)
+	{
+		return;
+	}
+
+	SchemaToConfigure->TemplateSequence->PathRelativeToRoot.Template = InPath;
+}
+
 void SCineAssemblySchemaWindow::ExpandTreeRecursive(TSharedPtr<FSchemaContentHierarchyFolder> FolderItem) const
 {
 	TreeView->SetItemExpansion(FolderItem, true);
@@ -1283,18 +1336,11 @@ void SCineAssemblySchemaWindow::DeleteTreeItem(TSharedPtr<ISchemaContentHierarch
 	if (TSharedPtr<FSchemaContentHierarchyFolder> FolderItem = TreeItem->AsFolder())
 	{
 		DeleteFolderItem(FolderItem);
-	}
-	else if (TSharedPtr<FSchemaContentHierarchyAsset> AssetItem = TreeItem->AsAsset())
-	{
-		DeleteAssetItem(AssetItem);
-	}
-
-	if (TSharedPtr<FSchemaContentHierarchyFolder> FolderItem = TreeItem->AsFolder())
-	{
 		TreeItem->ParentFolder->ChildFolders.Remove(FolderItem);
 	}
 	else if (TSharedPtr<FSchemaContentHierarchyAsset> AssetItem = TreeItem->AsAsset())
 	{
+		DeleteAssetItem(AssetItem);
 		TreeItem->ParentFolder->ChildAssets.Remove(AssetItem);
 	}
 
@@ -1307,21 +1353,19 @@ void SCineAssemblySchemaWindow::DeleteTreeItem(TSharedPtr<ISchemaContentHierarch
 
 void SCineAssemblySchemaWindow::DeleteFolderItem(TSharedPtr<FSchemaContentHierarchyFolder> FolderItem)
 {
-	// If the folder being deleted contains the TopLevelAssembly (which cannot be deleted), reparent the TopLevelAssembly first
-	if (ContainsRecursive(FolderItem, TopLevelAssemblyItem))
+	// If the folder contains the top level assembly, reparent that item to the root so it is not lost
+	if (FolderItem->ChildAssets.Contains(TopLevelAssemblyItem))
 	{
-		TopLevelAssemblyItem->ParentFolder = FolderItem->ParentFolder;
-		FolderItem->ParentFolder->ChildAssets.Add(TopLevelAssemblyItem);
-
 		FolderItem->ChildAssets.Remove(TopLevelAssemblyItem);
-
-		SchemaToConfigure->DefaultAssemblyPath = TopLevelAssemblyItem->ParentFolder->GetPath();
+		RootFolder->ChildAssets.Add(TopLevelAssemblyItem);
+		TopLevelAssemblyItem->ParentFolder = RootFolder;
 	}
 
 	for (const TSharedPtr<FSchemaContentHierarchyAsset>& ChildAsset : FolderItem->ChildAssets)
 	{
 		DeleteAssetItem(ChildAsset);
 	}
+
 	FolderItem->ChildAssets.Empty();
 
 	for (const TSharedPtr<FSchemaContentHierarchyFolder>& ChildFolder : FolderItem->ChildFolders)
@@ -1333,11 +1377,6 @@ void SCineAssemblySchemaWindow::DeleteFolderItem(TSharedPtr<FSchemaContentHierar
 
 void SCineAssemblySchemaWindow::DeleteAssetItem(TSharedPtr<FSchemaContentHierarchyAsset> AssetItem)
 {
-	if (AssetItem == TopLevelAssemblyItem)
-	{
-		return;
-	}
-
 	if (UMovieSceneSubAssemblySection* Section = AssetItem->SubAssemblySection.Get())
 	{
 		// Find the track that this section belongs to, and remove the section from that track. 
@@ -1353,9 +1392,13 @@ void SCineAssemblySchemaWindow::DeleteAssetItem(TSharedPtr<FSchemaContentHierarc
 	}
 }
 
-bool SCineAssemblySchemaWindow::ContainsRecursive(TSharedPtr<FSchemaContentHierarchyFolder> InParentFolder, TSharedPtr<FSchemaContentHierarchyAsset> InAssetItem)
+bool SCineAssemblySchemaWindow::ContainsRecursive(TSharedPtr<FSchemaContentHierarchyFolder> InParentFolder, TSharedPtr<ISchemaContentHierarchyTreeItem> InItem)
 {
-	if (InParentFolder->ChildAssets.Contains(InAssetItem))
+	if (InItem->AsAsset() && (InParentFolder->ChildAssets.Contains(InItem)))
+	{
+		return true;
+	}
+	else if (InItem->AsFolder() && (InParentFolder->ChildFolders.Contains(InItem)))
 	{
 		return true;
 	}
@@ -1363,7 +1406,7 @@ bool SCineAssemblySchemaWindow::ContainsRecursive(TSharedPtr<FSchemaContentHiera
 	bool bFoundRecursively = false;
 	for (const TSharedPtr<FSchemaContentHierarchyFolder>& ChildFolder : InParentFolder->ChildFolders)
 	{
-		bFoundRecursively |= ContainsRecursive(ChildFolder, InAssetItem);
+		bFoundRecursively |= ContainsRecursive(ChildFolder, InItem);
 	}
 
 	return bFoundRecursively;
@@ -1416,8 +1459,7 @@ FReply SCineAssemblySchemaWindow::OnCancelClicked()
 
 void SCineAssemblySchemaWindow::UpdateFolderList()
 {
-	// Update the schema's folder list
-	const TArray<FString> CachedFolderList = SchemaToConfigure->FoldersToCreate;
+	TArray<FString> CachedFolderList = GetFolderTemplates();
 
 	TArray<FString> FolderList;
 	GetFolderListRecursive(RootFolder, FolderList);
@@ -1425,8 +1467,10 @@ void SCineAssemblySchemaWindow::UpdateFolderList()
 	if (FolderList != CachedFolderList)
 	{
 		SchemaToConfigure->Modify();
-		SchemaToConfigure->FoldersToCreate = FolderList;
+		SetFolderTemplates(FolderList);
 	}
+
+	SetDefaultAssemblyPath(TopLevelAssemblyItem->ParentFolder->GetPath());
 }
 
 void SCineAssemblySchemaWindow::OnMovieSceneDataChanged(EMovieSceneDataChangeType InChangeType)
@@ -1673,6 +1717,16 @@ FReply SCineAssemblySchemaWindow::OnTreeRowAcceptDrop(const FDragDropEvent& InDr
 			DroppedItem = MakeShared<FSchemaContentHierarchyAsset>();
 		}
 
+		// Ensure that we are not attempting to drop an item onto itself or one of its children
+		if (DroppedItem == InItem)
+		{
+			return FReply::Unhandled();
+		}
+		else if (DroppedItem->AsFolder() && ContainsRecursive(DroppedItem->AsFolder(), InItem))
+		{
+			return FReply::Unhandled();
+		}
+
 		// If the dropped item is a folder, but the target folder already contains a subfolder with the same name, do not handle the drop (the dragged item will be reset to its original location)
 		const FString DroppedItemPath = TargetFolder->GetPath() / Operation->ItemName;
 		if (DroppedItem->AsFolder())
@@ -1690,11 +1744,6 @@ FReply SCineAssemblySchemaWindow::OnTreeRowAcceptDrop(const FDragDropEvent& InDr
 		{
 			TargetFolder->ChildAssets.Add(DroppedAsset);
 			TargetFolder->ChildAssets.Sort(UE::CineAssemblyTools::Private::SortTreeItems);
-
-			if (DroppedAsset == TopLevelAssemblyItem)
-			{
-				SchemaToConfigure->DefaultAssemblyPath = TopLevelAssemblyItem->ParentFolder->GetPath();
-			}
 		}
 		else if (TSharedPtr<FSchemaContentHierarchyFolder> DroppedFolder = DroppedItem->AsFolder())
 		{
@@ -1702,11 +1751,6 @@ FReply SCineAssemblySchemaWindow::OnTreeRowAcceptDrop(const FDragDropEvent& InDr
 
 			SetChildrenPathRecursive(DroppedFolder, DroppedFolder->GetPath());
 			TargetFolder->ChildFolders.Sort(UE::CineAssemblyTools::Private::SortTreeItems);
-
-			if (ContainsRecursive(DroppedFolder, TopLevelAssemblyItem))
-			{
-				SchemaToConfigure->DefaultAssemblyPath = TopLevelAssemblyItem->ParentFolder->GetPath();
-			}
 		}
 
 		UpdateFolderList();
diff --git a/Engine/Plugins/VirtualProduction/CinematicAssemblyTools/Source/CineAssemblyToolsEditor/Private/UI/CineAssembly/SCineAssemblySchemaWindow.h b/Engine/Plugins/VirtualProduction/CinematicAssemblyTools/Source/CineAssemblyToolsEditor/Private/UI/CineAssembly/SCineAssemblySchemaWindow.h
--- a/Engine/Plugins/VirtualProduction/CinematicAssemblyTools/Source/CineAssemblyToolsEditor/Private/UI/CineAssembly/SCineAssemblySchemaWindow.h
+++ b/Engine/Plugins/VirtualProduction/CinematicAssemblyTools/Source/CineAssemblyToolsEditor/Private/UI/CineAssembly/SCineAssemblySchemaWindow.h
@@ -173,6 +173,18 @@ private:
 	/** Build the content tree from the list of folders and assets saved in the schema properties */
 	void InitializeContentTree();
 
+	/** Get the list of template folder names from the Schema's Template Sequence */
+	TArray<FString> GetFolderTemplates();
+
+	/** Set the list of template folder names in the Schema's Template Sequence */
+	void SetFolderTemplates(TArray<FString> InFolderTemplates);
+
+	/** Get the default path of the Schema's Template Sequence */
+	FString GetDefaultAssemblyPath();
+
+	/** Set the default path of the Schema's Template Sequence */
+	void SetDefaultAssemblyPath(const FString& InPath);
+
 	/** Recursively expands every folder in the tree view */
 	void ExpandTreeRecursive(TSharedPtr<FSchemaContentHierarchyFolder> FolderItem) const;
 
@@ -221,8 +233,8 @@ private:
 	/** Generates a unique folder path name, assuming the input folder will be the parent folder */
 	FString MakeUniqueFolderPath(TSharedPtr<FSchemaContentHierarchyFolder> FolderItem);
 
-	/** Returns true if the input parent folder, or any of its child folders, contains in the input asset item */
-	bool ContainsRecursive(TSharedPtr<FSchemaContentHierarchyFolder> InParentFolder, TSharedPtr<FSchemaContentHierarchyAsset> InAssetItem);
+	/** Returns true if the input parent folder, or any of its child folders, contains in the input item */
+	bool ContainsRecursive(TSharedPtr<FSchemaContentHierarchyFolder> InParentFolder, TSharedPtr<ISchemaContentHierarchyTreeItem> InItem);
 
 	/** Update the list of folders to create from the current set in the tree view */
 	void UpdateFolderList();
