ShotManagement: Update CineAssembly "{parent}" token to better support parent Assembly objects that do not have valid packages yet.
When creating an Assembly that has SubAssemblies, the SubAssemblies might be named using the parent token, but parent Assembly will not have a valid asset name until configuration is done. This change to the token evaluation function now checks the package and can return the resolve AssemblyName property of the parent assembly.
Fixed CineAssembly to now properly assign the parent assembly immediately when creating SubAssemblies.
#rnx
#Synced-CL 51715153

diff --git a/Engine/Plugins/VirtualProduction/CinematicAssemblyTools/Source/CineAssemblyTools/Private/CineAssembly.cpp b/Engine/Plugins/VirtualProduction/CinematicAssemblyTools/Source/CineAssemblyTools/Private/CineAssembly.cpp
--- a/Engine/Plugins/VirtualProduction/CinematicAssemblyTools/Source/CineAssemblyTools/Private/CineAssembly.cpp
+++ b/Engine/Plugins/VirtualProduction/CinematicAssemblyTools/Source/CineAssemblyTools/Private/CineAssembly.cpp
@@ -356,6 +356,8 @@ UCineAssembly* UCineAssembly::CreateSubAssemblyFromTemplate(UObject* TemplateObj
 		SubAssembly->Initialize();
 	}
 
+	SubAssembly->ParentAssembly = this;
+
 	return SubAssembly;
 }
 
diff --git a/Engine/Plugins/VirtualProduction/CinematicAssemblyTools/Source/CineAssemblyTools/Private/CineAssemblyNamingTokens.cpp b/Engine/Plugins/VirtualProduction/CinematicAssemblyTools/Source/CineAssemblyTools/Private/CineAssemblyNamingTokens.cpp
--- a/Engine/Plugins/VirtualProduction/CinematicAssemblyTools/Source/CineAssemblyTools/Private/CineAssemblyNamingTokens.cpp
+++ b/Engine/Plugins/VirtualProduction/CinematicAssemblyTools/Source/CineAssemblyTools/Private/CineAssemblyNamingTokens.cpp
@@ -2,6 +2,7 @@
 
 #include "CineAssemblyNamingTokens.h"
 
+#include "Misc/PackageName.h"
 #include "MovieSceneMetaData.h"
 #include "NamingTokensEngineSubsystem.h"
 
@@ -68,7 +69,17 @@ void UCineAssemblyNamingTokens::OnCreateDefaultTokens(TArray<FNamingTokenData>&
 		{
 			if (InAssembly.IsValid() && InAssembly->ParentAssembly.IsValid())
 			{
-				return FText::FromString(InAssembly->ParentAssembly.GetAssetName());
+				const FString PackageName = InAssembly->ParentAssembly.GetLongPackageName();
+				if (FPackageName::IsInEngineTransientPackages(PackageName))
+				{
+					// If the parent assembly is transient, use the resolved AssemblyName value
+					const UCineAssembly* ParentAssembly = Cast<UCineAssembly>(InAssembly->ParentAssembly.ResolveObject());
+					return ParentAssembly ? ParentAssembly->AssemblyName.Resolved : FText::GetEmpty();
+				}
+				else
+				{
+					return FText::FromString(InAssembly->ParentAssembly.GetAssetName());
+				}
 			}
 			return FText::GetEmpty();
 		};
@@ -82,8 +93,6 @@ void UCineAssemblyNamingTokens::OnCreateDefaultTokens(TArray<FNamingTokenData>&
 			return FText::GetEmpty();
 		};
 
-
-
 	auto AuthorTokenFunc = [](TWeakObjectPtr<UCineAssembly> InAssembly) -> FText
 		{
 #if WITH_EDITOR
