ShotManagement: Fix the CineAssembly config window so that it properly sets a good default name for new Assemblies, and resets the name to a good default when no schema is selected.
Fix the assembly name textbox widget to immediately re-resolve tokens after the name is reset from a schema change.
#rnx
#Synced-CL 51715153

diff --git a/Engine/Plugins/VirtualProduction/CinematicAssemblyTools/Source/CineAssemblyToolsEditor/Private/UI/CineAssembly/SCineAssemblyConfigPanel.cpp b/Engine/Plugins/VirtualProduction/CinematicAssemblyTools/Source/CineAssemblyToolsEditor/Private/UI/CineAssembly/SCineAssemblyConfigPanel.cpp
--- a/Engine/Plugins/VirtualProduction/CinematicAssemblyTools/Source/CineAssemblyToolsEditor/Private/UI/CineAssembly/SCineAssemblyConfigPanel.cpp
+++ b/Engine/Plugins/VirtualProduction/CinematicAssemblyTools/Source/CineAssemblyToolsEditor/Private/UI/CineAssembly/SCineAssemblyConfigPanel.cpp
@@ -7,7 +7,6 @@
 #include "IStructureDetailsView.h"
 #include "Modules/ModuleManager.h"
 #include "PropertyEditorModule.h"
-#include "SNamingTokensEditableTextBox.h"
 #include "Widgets/Input/SSegmentedControl.h"
 
 #define LOCTEXT_NAMESPACE "SCineAssemblyConfigPanel"
@@ -92,7 +91,7 @@ void SCineAssemblyConfigPanel::Construct(const FArguments& InArgs, UCineAssembly
 								.FillWidth(1.0f)
 								.VAlign(VAlign_Center)
 								[
-									SNew(SNamingTokensEditableTextBox)
+									SAssignNew(AssemblyNameTextBox, SNamingTokensEditableTextBox)
 										.Contexts({ NamingTokenContext.Get() })
 										.FilterArgs(FilterArgs)
 										.EvaluationFrequency(1.0f)
@@ -144,6 +143,8 @@ bool SCineAssemblyConfigPanel::ValidateAssemblyName(const FText& InText, FText&
 
 void SCineAssemblyConfigPanel::Refresh()
 {
+	AssemblyNameTextBox->EvaluateNamingTokens();
+
 	// The details view needs to be redrawn to show the new metadata fields from the selected schema
 	DetailsView->ForceRefresh();
 
diff --git a/Engine/Plugins/VirtualProduction/CinematicAssemblyTools/Source/CineAssemblyToolsEditor/Private/UI/CineAssembly/SCineAssemblyConfigPanel.h b/Engine/Plugins/VirtualProduction/CinematicAssemblyTools/Source/CineAssemblyToolsEditor/Private/UI/CineAssembly/SCineAssemblyConfigPanel.h
--- a/Engine/Plugins/VirtualProduction/CinematicAssemblyTools/Source/CineAssemblyToolsEditor/Private/UI/CineAssembly/SCineAssemblyConfigPanel.h
+++ b/Engine/Plugins/VirtualProduction/CinematicAssemblyTools/Source/CineAssemblyToolsEditor/Private/UI/CineAssembly/SCineAssemblyConfigPanel.h
@@ -8,6 +8,7 @@
 #include "CineAssemblyNamingTokens.h"
 #include "IDetailsView.h"
 #include "NamingTokensEngineSubsystem.h"
+#include "SNamingTokensEditableTextBox.h"
 #include "Widgets/Layout/SWidgetSwitcher.h"
 #include "Widgets/Views/STreeView.h"
 
@@ -114,4 +115,7 @@ private:
 
 	/** Naming Token filter properties */
 	FNamingTokenFilterArgs FilterArgs;
+
+	/** TextBox displaying the name of the assembly being configured */
+	TSharedPtr<SNamingTokensEditableTextBox> AssemblyNameTextBox;
 };
diff --git a/Engine/Plugins/VirtualProduction/CinematicAssemblyTools/Source/CineAssemblyToolsEditor/Private/UI/CineAssembly/SCineAssemblyConfigWindow.cpp b/Engine/Plugins/VirtualProduction/CinematicAssemblyTools/Source/CineAssemblyToolsEditor/Private/UI/CineAssembly/SCineAssemblyConfigWindow.cpp
--- a/Engine/Plugins/VirtualProduction/CinematicAssemblyTools/Source/CineAssemblyToolsEditor/Private/UI/CineAssembly/SCineAssemblyConfigWindow.cpp
+++ b/Engine/Plugins/VirtualProduction/CinematicAssemblyTools/Source/CineAssemblyToolsEditor/Private/UI/CineAssembly/SCineAssemblyConfigWindow.cpp
@@ -2,6 +2,7 @@
 
 #include "SCineAssemblyConfigWindow.h"
 
+#include "AssetToolsModule.h"
 #include "AssetViewUtils.h"
 #include "CineAssemblyFactory.h"
 #include "CineAssemblyToolsStyle.h"
@@ -31,6 +32,12 @@ void SCineAssemblyConfigWindow::Construct(const FArguments& InArgs, const FStrin
 	CineAssemblyToConfigure = TStrongObjectPtr<UCineAssembly>(NewObject<UCineAssembly>(GetTransientPackage(), NAME_None, RF_Transient));
 	CineAssemblyToConfigure->Initialize();
 
+	// Initialize the Assembly's name to a good default
+	FAssetToolsModule& AssetToolsModule = FModuleManager::GetModuleChecked<FAssetToolsModule>("AssetTools");
+	const TOptional<FString> DefaultAssemblyName = AssetToolsModule.Get().GetDefaultAssetNameForClass(UCineAssembly::StaticClass());
+
+	CineAssemblyToConfigure->AssemblyName.Template = DefaultAssemblyName.Get(TEXT("NewCineAssembly"));
+
 	if (UMovieSceneMetaData* MetaDataToConfigure = CineAssemblyToConfigure->FindOrAddMetaData<UMovieSceneMetaData>())
 	{
 		// Default the MetaData for the CineAssembly being configured.
@@ -205,6 +212,13 @@ void SCineAssemblyConfigWindow::OnSchemaSelected(const FAssetData& AssetData)
 	{
 		CineAssemblyToConfigure->ChangeSchema(nullptr);
 		CineAssemblyToConfigure->Initialize();
+
+		// Re-initialize the Assembly's name to a good default
+		FAssetToolsModule& AssetToolsModule = FModuleManager::GetModuleChecked<FAssetToolsModule>("AssetTools");
+		const TOptional<FString> DefaultAssemblyName = AssetToolsModule.Get().GetDefaultAssetNameForClass(UCineAssembly::StaticClass());
+
+		CineAssemblyToConfigure->AssemblyName.Template = DefaultAssemblyName.Get(TEXT("NewCineAssembly"));
+
 		ConfigPanel->Refresh();
 	}
 	else if (UCineAssemblySchema* CineAssemblySchema = Cast<UCineAssemblySchema>(AssetData.GetAsset()))
