ShotManagement: Move the assembly asset tree view into its own widget class to make it easier to re-use across tools. The primary motivation is to use the same tree view when creating CineAssemblies as we show when creating CineAssemblySchemas.
Remove all of the custom tree view implementation from the CineAssembly creation tool and the CineAssemblySchema creation tool and just use the new SCineAssemblyAssetTreeView widget instead.
Some refactoring of the asset tree view compared to the original implementation to better support read-only and read-write use cases.

#rnx
#Synced-CL 51715153

diff --git a/Engine/Plugins/VirtualProduction/CinematicAssemblyTools/Source/CineAssemblyToolsEditor/Private/UI/CineAssembly/SCineAssemblyAssetTreeView.cpp b/Engine/Plugins/VirtualProduction/CinematicAssemblyTools/Source/CineAssemblyToolsEditor/Private/UI/CineAssembly/SCineAssemblyAssetTreeView.cpp
new file mode 100644
--- /dev/null
+++ b/Engine/Plugins/VirtualProduction/CinematicAssemblyTools/Source/CineAssemblyToolsEditor/Private/UI/CineAssembly/SCineAssemblyAssetTreeView.cpp
@@ -0,0 +1,1128 @@
+// Copyright Epic Games, Inc. All Rights Reserved.
+
+#include "SCineAssemblyAssetTreeView.h"
+
+#include "Algo/Contains.h"
+#include "CineAssembly.h"
+#include "CineAssemblyNamingTokens.h"
+#include "CineAssemblyToolsStyle.h"
+#include "Framework/MultiBox/MultiBoxBuilder.h"
+#include "MovieScene.h"
+#include "MovieSceneCommonHelpers.h"
+#include "MovieSceneSubAssemblySection.h"
+#include "MovieSceneSubAssemblyTrack.h"
+#include "PropertyCustomizationHelpers.h"
+#include "Widgets/Layout/SBorder.h"
+#include "Widgets/SBoxPanel.h"
+#include "Widgets/Text/STextBlock.h"
+
+#define LOCTEXT_NAMESPACE "SCineAssemblyAssetTreeView"
+
+struct ICineAssemblyTreeItem : TSharedFromThis<ICineAssemblyTreeItem>
+{
+	virtual ~ICineAssemblyTreeItem() {}
+
+	virtual TSharedPtr<FCineAssemblyFolderItem> AsFolder() { return nullptr; }
+	virtual TSharedPtr<FSubAssemblyItem> AsSubAssembly() { return nullptr; }
+	virtual TSharedPtr<FSubAssemblySectionItem> AsSubAssemblySection() { return nullptr; }
+
+	/** Returns true if this item should never be editable, regardless of the widget's read-only state */
+	virtual bool IsReadOnly() const { return false; }
+
+	/** Returns true if this item can be dragged to a new location in the tree */
+	virtual bool CanDrag() const { return true; }
+
+	/** Returns the assembly associated with this tree item */
+	UCineAssembly* GetAssembly() { return AssociatedAssembly.Get(); }
+
+	/** Sets the assembly associated with this tree item */
+	void SetAssembly(UCineAssembly* InAssembly) { AssociatedAssembly = InAssembly; }
+
+	/** Whether this item can accept an item being dropped onto it */
+	virtual TOptional<EItemDropZone> CanAcceptDrop() const { return TOptional<EItemDropZone>(); }
+
+	/** Return the name of the tree item */
+	virtual FString GetName() const = 0;
+
+	/** Set the name of the tree item */
+	virtual void SetName(const FString& InName) = 0;
+
+	/** Return the display name of the tree item */
+	virtual FText GetDisplayName() const = 0;
+
+	/** Set the display name of the tree item */
+	virtual void SetDisplayName(const FText& InDisplayName) = 0;
+
+	/** Return the path of the tree item (relative to the root folder) */
+	virtual FString GetPath() const = 0;
+
+	/** Set the path of the tree item (relative to the root folder) */
+	virtual void SetPath(const FString& InPath) = 0;
+
+	/** The parent folder of this item */
+	TSharedPtr<FCineAssemblyFolderItem> ParentFolder;
+
+	/** The text widget that displays this item's name and supports renaming */
+	TSharedPtr<SNamingTokensEditableTextBox> NameWidget;
+
+	/** The assembly that owns or is represented by this tree item */
+	TWeakObjectPtr<UCineAssembly> AssociatedAssembly;
+
+	/** Context object used to evaluate naming tokens displayed in this item's name widget */
+	TStrongObjectPtr<UCineAssemblyNamingTokensContext> NamingTokenContext;
+};
+
+struct FSubAssemblyItem : ICineAssemblyTreeItem
+{
+	virtual TSharedPtr<FSubAssemblyItem> AsSubAssembly() override { return SharedThis(this); }
+
+	//~ Begin ICineAssemblyTreeItem Interface
+	virtual FString GetName() const override
+	{
+		return AssociatedAssembly.IsValid() ? AssociatedAssembly->AssemblyName.Template : FString();
+	}
+
+	virtual FText GetDisplayName() const override
+	{
+		return AssociatedAssembly.IsValid() ? FText::FromString(AssociatedAssembly->AssemblyName.Template) : FText();
+	}
+
+	virtual FString GetPath() const override
+	{
+		return AssociatedAssembly.IsValid() ? AssociatedAssembly->PathRelativeToRoot.Template : FString();
+	}
+
+	virtual void SetName(const FString& InName) override { ensureAlways(false); }
+	virtual void SetDisplayName(const FText& InDisplayName) override { ensureAlways(false); }
+	virtual void SetPath(const FString& InPath) override
+	{
+		if (AssociatedAssembly.IsValid())
+		{
+			AssociatedAssembly->PathRelativeToRoot.Template = FPaths::GetPath(InPath);
+		}
+	}
+	virtual bool IsReadOnly() const override { return true; }
+	//~ End ICineAssemblyTreeItem Interface
+};
+
+struct FSubAssemblySectionItem : ICineAssemblyTreeItem
+{
+	virtual TSharedPtr<FSubAssemblySectionItem> AsSubAssemblySection() override { return SharedThis(this); }
+
+	//~ Begin ICineAssemblyTreeItem Interface
+	virtual FString GetName() const override
+	{
+		return SubAssemblySection.IsValid() ? SubAssemblySection->GetSequenceName().ToString() : FString();
+	}
+
+	virtual void SetName(const FString& InName) override
+	{
+		if (SubAssemblySection.IsValid())
+		{
+			SubAssemblySection->SetSequenceName(FText::FromString(InName));
+		}
+	}
+
+	virtual FText GetDisplayName() const override
+	{
+		return SubAssemblySection.IsValid() ? SubAssemblySection->GetSequenceName() : FText();
+	}
+
+	virtual void SetDisplayName(const FText& InDisplayName) override
+	{
+		if (SubAssemblySection.IsValid())
+		{
+			SubAssemblySection->SetSequenceName(InDisplayName);
+		}
+	}
+
+	virtual FString GetPath() const override
+	{
+		return SubAssemblySection.IsValid() ? SubAssemblySection->GetSequencePath() : FString();
+	}
+
+	virtual void SetPath(const FString& InPath) override
+	{
+		if (SubAssemblySection.IsValid())
+		{
+			SubAssemblySection->SetSequencePath(FPaths::GetPath(InPath));
+		}
+	}
+	//~ End ICineAssemblyTreeItem Interface
+
+	/** The MovieScene section represented in the content tree by this asset item */
+	TWeakObjectPtr<UMovieSceneSubAssemblySection> SubAssemblySection;
+};
+
+struct FCineAssemblyFolderItem : ICineAssemblyTreeItem
+{
+	//~ Begin ICineAssemblyTreeItem Interface
+	virtual TSharedPtr<FCineAssemblyFolderItem> AsFolder() override { return SharedThis(this); }
+	virtual TOptional<EItemDropZone> CanAcceptDrop() const override { return EItemDropZone::OntoItem; }
+	virtual bool IsReadOnly() const override { return ParentFolder == nullptr; }
+	virtual bool CanDrag() const override { return ParentFolder != nullptr; }
+
+	virtual FString GetName() const override
+	{
+		return FPaths::GetPathLeaf(FullPath);
+	}
+
+	virtual void SetName(const FString& InName) override
+	{
+		FullPath = FPaths::GetPath(FullPath) / InName;
+	}
+
+	virtual FText GetDisplayName() const override
+	{
+		return FText::FromString(GetName());
+	}
+
+	virtual void SetDisplayName(const FText& InDisplayName) override
+	{
+		FullPath = FPaths::GetPath(FullPath) / InDisplayName.ToString();
+	}
+
+	virtual FString GetPath() const override
+	{
+		return FullPath;
+	}
+
+	virtual void SetPath(const FString& InPath) override
+	{
+		FullPath = InPath;
+	}
+	//~ End ICineAssemblyTreeItem Interface
+
+	TArray<TSharedPtr<FCineAssemblyFolderItem>> GetChildFolders() const
+	{
+		TArray<TSharedPtr<FCineAssemblyFolderItem>> ChildFolders;
+		for (const TSharedPtr<ICineAssemblyTreeItem>& Child : Children)
+		{
+			if (TSharedPtr<FCineAssemblyFolderItem> ChildFolder = Child->AsFolder())
+			{
+				ChildFolders.Add(ChildFolder);
+			}
+		}
+
+		return ChildFolders;
+	}
+
+	/** The fully qualified path of this item (relative to the root folder) */
+	FString FullPath;
+
+	/** The children of this folder item */
+	TArray<TSharedPtr<ICineAssemblyTreeItem>> Children;
+};
+
+namespace UE::CineAssemblyTools::Private
+{
+	auto SortTreeItems = [](const TSharedPtr<ICineAssemblyTreeItem>& A, const TSharedPtr<ICineAssemblyTreeItem>& B)
+	{
+		const bool bAIsFolder = A->AsFolder().IsValid();
+		const bool bBIsFolder = B->AsFolder().IsValid();
+		if (bAIsFolder != bBIsFolder)
+		{
+			return bBIsFolder;
+		}
+
+		return A->GetName() < B->GetName();
+	};
+
+	void SortTreeItemsRecursive(const TSharedPtr<FCineAssemblyFolderItem>& FolderItem)
+	{
+		FolderItem->Children.Sort(SortTreeItems);
+		for (const TSharedPtr<FCineAssemblyFolderItem>& ChildFolder : FolderItem->GetChildFolders())
+		{
+			SortTreeItemsRecursive(ChildFolder);
+		}
+	}
+
+	/** Get the list of template folder names for the input Assembly */
+	TArray<FString> GetFolderTemplates(UCineAssembly* Assembly)
+	{
+		if (!Assembly)
+		{
+			return TArray<FString>();
+		}
+
+		// Extract the template folder names from the input assembly's default folder list
+		TArrayView<FTemplateString> DefaultFolderNames = Assembly->DefaultFolderNames;
+
+		TArray<FString> TemplateFolderNames;
+		TemplateFolderNames.Reserve(DefaultFolderNames.Num());
+
+		Algo::Transform(DefaultFolderNames, TemplateFolderNames, [](const FTemplateString& TemplateString) { return TemplateString.Template; });
+
+		return TemplateFolderNames;
+	}
+
+	/** Set the list of template folder names for the input Assembly */
+	void SetFolderTemplates(UCineAssembly* Assembly, const TArray<FString>& InFolderTemplates)
+	{
+		if (!Assembly)
+		{
+			return;
+		}
+
+		TArray<FTemplateString>& DefaultFolderNames = Assembly->DefaultFolderNames;
+
+		DefaultFolderNames.Reset(InFolderTemplates.Num());
+		Algo::Transform(InFolderTemplates, DefaultFolderNames, [](const FString& TemplateString) { return FTemplateString(TemplateString, FText::GetEmpty()); });
+	}
+
+	/** Recursively get the path of every folder item in the tree view */
+	void GetFolderListRecursive(const TSharedPtr<FCineAssemblyFolderItem>& FolderItem, TArray<FString>& OutFolderList)
+	{
+		for (const TSharedPtr<FCineAssemblyFolderItem>& ChildFolder : FolderItem->GetChildFolders())
+		{
+			OutFolderList.Add(ChildFolder->GetPath());
+			GetFolderListRecursive(ChildFolder, OutFolderList);
+		}
+	}
+
+	/** Returns the folder in the tree whose path matches the input path */
+	TSharedPtr<FCineAssemblyFolderItem> FindFolderByPathRecursive(const TSharedPtr<FCineAssemblyFolderItem>& FolderItem, const FString& FolderPath)
+	{
+		if (FolderItem->GetPath().Equals(FolderPath))
+		{
+			return FolderItem;
+		}
+
+		for (const TSharedPtr<FCineAssemblyFolderItem>& ChildFolder : FolderItem->GetChildFolders())
+		{
+			if (const TSharedPtr<FCineAssemblyFolderItem>& MatchingFolder = FindFolderByPathRecursive(ChildFolder, FolderPath))
+			{
+				return MatchingFolder;
+			}
+		}
+
+		return nullptr;
+	}
+
+	/** Modifies the path of the input folder item and all of its children (recursively) with the input new path */
+	void SetChildrenPathRecursive(const TSharedPtr<FCineAssemblyFolderItem>& FolderItem, const FString& NewPath)
+	{
+		for (const TSharedPtr<ICineAssemblyTreeItem>& Child : FolderItem->Children)
+		{
+			const FString NewChildPath = NewPath / Child->GetName();
+			Child->SetPath(NewChildPath);
+
+			if (TSharedPtr<FCineAssemblyFolderItem> ChildFolder = Child->AsFolder())
+			{
+				SetChildrenPathRecursive(ChildFolder, NewChildPath);
+			}
+		}
+	}
+
+	/** Returns true if the input parent folder, or any of its child folders, contains the input item */
+	bool ContainsRecursive(const TSharedPtr<FCineAssemblyFolderItem>& InParentFolder, const TSharedPtr<ICineAssemblyTreeItem>& InItem)
+	{
+		if (InParentFolder->Children.Contains(InItem))
+		{
+			return true;
+		}
+
+		for (const TSharedPtr<FCineAssemblyFolderItem>& ChildFolder : InParentFolder->GetChildFolders())
+		{
+			if (ContainsRecursive(ChildFolder, InItem))
+			{
+				return true;
+			}
+		}
+
+		return false;
+	}
+
+	/** Removes a SubAssembly section from its owning Track (and remove the Track from its owning MovieScene if it would then be empty) */
+	void RemoveSubAssemblySection(UMovieSceneSubAssemblySection* SubAssemblySection)
+	{
+		if (!SubAssemblySection)
+		{
+			return;
+		}
+
+		// Find the track that this section belongs to, and remove the section from that track. 
+		// IF the track has no other sections after deleting the section, then remove the whole track.
+		if (UMovieSceneSubAssemblyTrack* Track = SubAssemblySection->GetTypedOuter<UMovieSceneSubAssemblyTrack>())
+		{
+			Track->RemoveSection(*SubAssemblySection);
+			if (Track->GetAllSections().IsEmpty())
+			{
+				if (UMovieScene* OwningMovieScene = Track->GetTypedOuter<UMovieScene>())
+				{
+					OwningMovieScene->RemoveTrack(*Track);
+				}
+			}
+		}
+	}
+
+	/** Generates a unique folder path name, assuming the input folder will be the parent folder */
+	FString MakeUniqueFolderPath(const TSharedPtr<FCineAssemblyFolderItem>& FolderItem)
+	{
+		// This implementation is based on a similar utility in AssetTools for creating a unique asset name.
+		const FString BaseName = FolderItem->GetPath() / TEXT("NewFolder");
+
+		// Find the index in the string of the last non-numeric character
+		int32 CharIndex = BaseName.Len() - 1;
+		while (CharIndex >= 0 && BaseName[CharIndex] >= TEXT('0') && BaseName[CharIndex] <= TEXT('9'))
+		{
+			--CharIndex;
+		}
+
+		// Trim the numeric characters off the end of the BaseName string, but remember the integer that was trimmed off to increment and append to the output
+		int32 IntSuffix = 1;
+		FString TrimmedBaseName = BaseName;
+		if (CharIndex >= 0 && CharIndex < BaseName.Len() - 1)
+		{
+			TrimmedBaseName = BaseName.Left(CharIndex + 1);
+
+			const FString TrailingInteger = BaseName.RightChop(CharIndex + 1);
+			IntSuffix = FCString::Atoi(*TrailingInteger) + 1;
+		}
+
+		FString WorkingName = TrimmedBaseName;
+
+		while (Algo::ContainsBy(FolderItem->GetChildFolders(), WorkingName, &FCineAssemblyFolderItem::FullPath))
+		{
+			WorkingName = FString::Printf(TEXT("%s%d"), *TrimmedBaseName, IntSuffix);
+			IntSuffix++;
+		}
+
+		return WorkingName;
+	}
+}
+
+void SCineAssemblyAssetTreeView::Construct(const FArguments& InArgs, UCineAssembly* InAssembly)
+{
+	WeakAssembly = InAssembly;
+
+	bIsReadOnly = InArgs._IsReadOnly;
+	bDisplayHintText = InArgs._DisplayHintText;
+	SelectionMode = bIsReadOnly ? ESelectionMode::None : InArgs._SelectionMode;
+	bShouldEvaluateTokens = InArgs._ShouldEvaluateTokens;
+
+	OnItemRemoved = InArgs._OnItemRemoved;
+
+	InitializeContentTree(InAssembly, RootFolder, TopLevelAssemblyItem);
+	TreeItems.Add(RootFolder);
+
+	TreeView = SNew(STreeView<TSharedPtr<ICineAssemblyTreeItem>>)
+		.TreeItemsSource(&TreeItems)
+		.SelectionMode(SelectionMode)
+		.OnGenerateRow(this, &SCineAssemblyAssetTreeView::OnGenerateTreeRow)
+		.OnGetChildren(this, &SCineAssemblyAssetTreeView::OnGetChildren)
+		.OnItemsRebuilt(this, &SCineAssemblyAssetTreeView::OnTreeItemsRebuilt)
+		.OnContextMenuOpening(this, &SCineAssemblyAssetTreeView::MakeContentTreeContextMenu)
+		.OnMouseButtonDoubleClick(this, &SCineAssemblyAssetTreeView::OnTreeViewDoubleClick)
+		.OnKeyDownHandler(this, &SCineAssemblyAssetTreeView::OnTreeViewKeyDown);
+
+	ExpandTreeRecursive(RootFolder);
+
+	ChildSlot
+	[
+		TreeView.ToSharedRef()
+	];
+}
+
+void SCineAssemblyAssetTreeView::InitializeContentTree(UCineAssembly* Assembly, TSharedPtr<FCineAssemblyFolderItem>& ContentRoot, TSharedPtr<FSubAssemblyItem>& TopLevelAssembly)
+{
+	using namespace UE::CineAssemblyTools::Private;
+
+	// Create the content tree root
+	ContentRoot = MakeShared<FCineAssemblyFolderItem>();
+	ContentRoot->SetPath(TEXT(""));
+
+	if (!Assembly)
+	{
+		return;
+	}
+
+	// Sort the list of folders so that paths are added to the tree in the proper order
+	TArray<FString> FolderTemplates = GetFolderTemplates(Assembly);
+	FolderTemplates.Sort();
+
+	for (const FString& FolderPath : FolderTemplates)
+	{
+		const FString ParentPath = FPaths::GetPath(FolderPath);
+
+		// Walk the tree until we find an item whose path matches the parent path. The new tree item will be created as one of its children
+		TSharedPtr<FCineAssemblyFolderItem> ParentFolder = FindFolderByPathRecursive(ContentRoot, ParentPath);
+		if (ParentFolder)
+		{
+			TSharedPtr<FCineAssemblyFolderItem> NewItem = MakeShared<FCineAssemblyFolderItem>();
+			NewItem->SetPath(FolderPath);
+			NewItem->ParentFolder = ParentFolder;
+			NewItem->SetAssembly(Assembly);
+			ParentFolder->Children.Add(NewItem);
+		}
+	}
+
+	// Add the top-level assembly node
+	TopLevelAssembly = MakeShared<FSubAssemblyItem>();
+
+	const FString& TopLevelAssemblyPath = Assembly->PathRelativeToRoot.Template;
+	TSharedPtr<FCineAssemblyFolderItem> ParentFolder = FindFolderByPathRecursive(ContentRoot, TopLevelAssemblyPath);
+	if (ParentFolder)
+	{
+		TopLevelAssembly->ParentFolder = ParentFolder;
+		TopLevelAssembly->SetAssembly(Assembly);
+		ParentFolder->Children.Add(TopLevelAssembly);
+	}
+
+	// Add an item to the tree for each SubAssemblySection in the Assembly's MovieScene
+	AddSubAssemblySectionsToTree(Assembly, ContentRoot);
+
+	// Add an item to the tree for each child SubAssembly of this Assembly
+	AddSubAssembliesToTree(Assembly, ContentRoot);
+
+	SortTreeItemsRecursive(ContentRoot);
+}
+
+void SCineAssemblyAssetTreeView::AddSubAssemblySectionsToTree(UCineAssembly* Assembly, TSharedPtr<FCineAssemblyFolderItem>& ContentRoot)
+{
+	TArray<UMovieSceneSubSection*> SubSections;
+	MovieSceneHelpers::GetDescendantSubSections(Assembly->GetMovieScene(), SubSections);
+
+	for (UMovieSceneSubSection* SubSection : SubSections)
+	{
+		UMovieSceneSubAssemblySection* SubAssemblySection = Cast<UMovieSceneSubAssemblySection>(SubSection);
+
+		if (SubAssemblySection && SubAssemblySection->IsTemplateSection())
+		{
+			// Walk the tree until we find a parent folder whose path matches the path of the SubAssembly section
+			if (TSharedPtr<FCineAssemblyFolderItem> ParentFolder = UE::CineAssemblyTools::Private::FindFolderByPathRecursive(ContentRoot, SubAssemblySection->GetSequencePath()))
+			{
+				// Make a new tree item for this section and add it into the tree
+				TSharedPtr<FSubAssemblySectionItem> NewItem = MakeShared<FSubAssemblySectionItem>();
+				NewItem->ParentFolder = ParentFolder;
+				NewItem->SubAssemblySection = SubAssemblySection;
+				ParentFolder->Children.Add(NewItem);
+			}
+		}
+	}
+}
+
+void SCineAssemblyAssetTreeView::AddSubAssembliesToTree(UCineAssembly* Assembly, TSharedPtr<FCineAssemblyFolderItem>& ContentRoot)
+{
+	// Each SubAssembly could have its own content tree that it will create.
+	// To display the full tree, we build a temporary content tree for the SubAssembly, then re-parent its root-level items into the parent tree.
+	for (UMovieSceneSubSection* SubSection : Assembly->SubAssemblies)
+	{
+		if (UCineAssembly* SubAssembly = Cast<UCineAssembly>(SubSection->GetSequence()))
+		{
+			SubAssembly->AssemblyName.Resolved = UCineAssemblyNamingTokens::GetResolvedText(SubAssembly->AssemblyName.Template, SubAssembly);
+
+			// Generate the content tree for this SubAssembly
+			TSharedPtr<FCineAssemblyFolderItem> SubAssemblyContentRoot;
+			TSharedPtr<FSubAssemblyItem> SubAssemblyTreeItem;
+			InitializeContentTree(SubAssembly, SubAssemblyContentRoot, SubAssemblyTreeItem);
+
+			// Reparent the items from the SubAssembly content tree into the parent assembly's content tree
+			if (TSharedPtr<FCineAssemblyFolderItem> ParentFolder = UE::CineAssemblyTools::Private::FindFolderByPathRecursive(ContentRoot, SubAssembly->PathRelativeToParent.Template))
+			{
+				for (const TSharedPtr<ICineAssemblyTreeItem>& Child : SubAssemblyContentRoot->Children)
+				{
+					ParentFolder->Children.Add(Child);
+					Child->ParentFolder = ParentFolder;
+				}
+			}
+		}
+	}
+}
+
+void SCineAssemblyAssetTreeView::Reinitialize()
+{
+	TreeItems.Empty();
+	InitializeContentTree(WeakAssembly.Get(), RootFolder, TopLevelAssemblyItem);
+	TreeItems.Add(RootFolder);
+
+	ExpandTreeRecursive(RootFolder);
+	TreeView->RequestTreeRefresh();
+}
+
+void SCineAssemblyAssetTreeView::AddFolder()
+{
+	using namespace UE::CineAssemblyTools::Private;
+
+	TSharedPtr<FCineAssemblyFolderItem> NewFolder = MakeShared<FCineAssemblyFolderItem>();
+
+	// Get the parent item for the new folder being added (this can be the root folder if no parent is currently selected)
+	// The TreeView uses single selection mode, so at most one item can ever be selected by the user
+	TSharedPtr<FCineAssemblyFolderItem> ParentFolder = RootFolder;
+
+	TArray<TSharedPtr<ICineAssemblyTreeItem>> SelectedItems = TreeView->GetSelectedItems();
+	if (SelectedItems.Num() == 1)
+	{
+		if (TSharedPtr<FCineAssemblyFolderItem> SelectedFolder = SelectedItems[0]->AsFolder())
+		{
+			ParentFolder = SelectedFolder;
+		}
+		else
+		{
+			ParentFolder = SelectedItems[0]->ParentFolder;
+		}
+	}
+
+	NewFolder->SetPath(MakeUniqueFolderPath(ParentFolder));
+	NewFolder->ParentFolder = ParentFolder;
+	NewFolder->SetAssembly(WeakAssembly.Get());
+
+	ParentFolder->Children.Add(NewFolder);
+
+	// Sort the children alphabetically to maintain a good ordering with the new folder
+	ParentFolder->Children.Sort(SortTreeItems);
+
+	// Save a reference to this item so that when the tree is rebuilt, we can immediately start editing its name
+	MostRecentlyAddedItem = NewFolder;
+
+	UpdateFolderList();
+
+	TreeView->SetItemExpansion(ParentFolder, true);
+	TreeView->RequestTreeRefresh();
+}
+
+TSharedRef<ITableRow> SCineAssemblyAssetTreeView::OnGenerateTreeRow(TSharedPtr<ICineAssemblyTreeItem> TreeItem, const TSharedRef<STableViewBase>& OwnerTable)
+{
+	TSharedPtr<SImage> Icon = SNew(SImage);
+
+	if (TreeItem->AsFolder())
+	{
+		Icon->SetImage(FCineAssemblyToolsStyle::Get().GetBrush("Icons.Folder"));
+		Icon->SetColorAndOpacity(FAppStyle::Get().GetSlateColor("ContentBrowser.DefaultFolderColor"));
+	}
+	else
+	{
+		Icon->SetImage(FCineAssemblyToolsStyle::Get().GetBrush("Icons.Sequencer"));
+		Icon->SetColorAndOpacity(FLinearColor::White);
+	}
+
+	if (!TreeItem->NamingTokenContext)
+	{
+		TreeItem->NamingTokenContext = TStrongObjectPtr<UCineAssemblyNamingTokensContext>(NewObject<UCineAssemblyNamingTokensContext>());
+	}
+	TreeItem->NamingTokenContext->Assembly = TreeItem->GetAssembly();
+
+	FNamingTokenFilterArgs FilterArgs;
+	FilterArgs.AdditionalNamespacesToInclude.Add(UCineAssemblyNamingTokens::TokenNamespace);
+
+	return SNew(STableRow<TSharedPtr<ICineAssemblyTreeItem>>, OwnerTable)
+		.ShowSelection(true)
+		.Padding(FMargin(8.0f, 0.0f, 8.0f, 0.0f))
+		.OnCanAcceptDrop_Lambda([](const FDragDropEvent& InDragDropEvent, EItemDropZone InDropZone, TSharedPtr<ICineAssemblyTreeItem> InItem) { return InItem->CanAcceptDrop(); })
+		.OnAcceptDrop(this, &SCineAssemblyAssetTreeView::OnTreeRowAcceptDrop)
+		.OnDragDetected(this, &SCineAssemblyAssetTreeView::OnTreeRowDragDetected)
+		.Content()
+		[
+			SNew(SHorizontalBox)
+
+				+ SHorizontalBox::Slot()
+				.AutoWidth()
+				.VAlign(VAlign_Center)
+				.Padding(0.0f, 0.0f, 2.0f, 0.0f)
+				[
+					Icon.ToSharedRef()
+				]
+
+				+ SHorizontalBox::Slot()
+				.AutoWidth()
+				.VAlign(VAlign_Center)
+				.Padding(0.0f, 0.0f, 8.0f, 0.0f)
+				[
+					SAssignNew(TreeItem->NameWidget, SNamingTokensEditableTextBox)
+						.SupportInlineEdit(true)
+						.ShouldEvaluateTokens(bShouldEvaluateTokens)
+						.Contexts({ TreeItem->NamingTokenContext.Get() })
+						.FilterArgs(FilterArgs)
+						.EvaluationFrequency(1.0f)
+						.TextBoxPadding(FMargin(4, 2, 0, 2))
+						.IsReadOnly(bIsReadOnly || TreeItem->IsReadOnly())
+						.DisplayTokenIcon(false)
+						.DisplayBorderImage(false)
+						.DisplayErrorMessage(!TreeItem->IsReadOnly())
+						.Text(this, &SCineAssemblyAssetTreeView::GetTreeItemDisplayName, TreeItem)
+						.OnTextCommitted(this, &SCineAssemblyAssetTreeView::OnTreeItemTextCommitted, TreeItem)
+						.OnValidateTokenizedText(this, &SCineAssemblyAssetTreeView::OnValidateTreeItemName, TreeItem)
+						.OnTokenizedTextEvaluated(this, &SCineAssemblyAssetTreeView::OnTreeItemTextResolved, TreeItem)
+				]
+
+			+ SHorizontalBox::Slot()
+				.AutoWidth()
+				.VAlign(VAlign_Center)
+				[
+					SNew(STextBlock)
+						.Text(this, &SCineAssemblyAssetTreeView::GetTreeItemHintText, TreeItem)
+						.Font(FCoreStyle::GetDefaultFontStyle("Italic", 10))
+						.ColorAndOpacity(FSlateColor(FLinearColor(1.0f, 1.0f, 1.0f, 0.3f)))
+				]
+		];
+}
+
+void SCineAssemblyAssetTreeView::OnGetChildren(TSharedPtr<ICineAssemblyTreeItem> TreeItem, TArray<TSharedPtr<ICineAssemblyTreeItem>>& OutNodes)
+{
+	// Children are pre-sorted by SortTreeItems (assets first, then folders, alphabetical within each group)
+	if (TSharedPtr<FCineAssemblyFolderItem> FolderItem = TreeItem->AsFolder())
+	{
+		OutNodes.Append(FolderItem->Children);
+	}
+}
+
+void SCineAssemblyAssetTreeView::OnTreeItemsRebuilt()
+{
+	// Upon regenerating the tree view, allow the user to immediately interact with the name widget of the newly added folder in order to rename it
+	if (MostRecentlyAddedItem && MostRecentlyAddedItem->NameWidget)
+	{
+		MostRecentlyAddedItem->NameWidget->EnterEditingMode();
+		MostRecentlyAddedItem.Reset();
+	}
+}
+
+FText SCineAssemblyAssetTreeView::GetTreeItemDisplayName(TSharedPtr<ICineAssemblyTreeItem> TreeItem) const
+{
+	if (TreeItem == RootFolder)
+	{
+		return LOCTEXT("RootPathName", "Root Folder");
+	}
+	else if (TreeItem == TopLevelAssemblyItem)
+	{
+		return FText::FromString(TEXT("{assembly}"));
+	}
+	else
+	{
+		return TreeItem->GetDisplayName();
+	}
+}
+
+FText SCineAssemblyAssetTreeView::GetTreeItemHintText(TSharedPtr<ICineAssemblyTreeItem> TreeItem) const
+{
+	// Template origin text is always shown for SubAssemblySections regardless of bDisplayHintText,
+	// because it provides essential context about which template the section was created from.
+	if (TSharedPtr<FSubAssemblySectionItem> SectionItem = TreeItem->AsSubAssemblySection())
+	{
+		if (UMovieSceneSubAssemblySection* SubAssemblySection = SectionItem->SubAssemblySection.Get())
+		{
+			if (UObject* TemplateObject = SubAssemblySection->GetAssemblyTemplate())
+			{
+				return FText::Format(LOCTEXT("TemplateOriginText", "from {0}"), FText::FromString(TemplateObject->GetName()));
+			}
+		}
+	}
+
+	if (!bDisplayHintText)
+	{
+		return FText();
+	}
+
+	UCineAssembly* AssociatedAssembly = TreeItem->GetAssembly();
+	if (!AssociatedAssembly)
+	{
+		return FText();
+	}
+
+	if (TreeItem->AsSubAssembly())
+	{
+		if (UCineAssembly* ParentAssembly = Cast<UCineAssembly>(AssociatedAssembly->ParentAssembly.TryLoad()))
+		{
+			return FText::Format(LOCTEXT("CreatedByText", "created by {0}"), ParentAssembly->AssemblyName.Resolved);
+		}
+		else
+		{
+			return FText();
+		}
+	}
+
+	return FText::Format(LOCTEXT("CreatedByText", "created by {0}"), AssociatedAssembly->AssemblyName.Resolved);
+}
+
+void SCineAssemblyAssetTreeView::OnTreeItemTextCommitted(const FText& InText, ETextCommit::Type InCommitType, TSharedPtr<ICineAssemblyTreeItem> TreeItem)
+{
+	using namespace UE::CineAssemblyTools::Private;
+
+	// Early-out if the name has not actually changed
+	const FString OldName = TreeItem->GetName();
+	if (OldName.Equals(InText.ToString()))
+	{
+		return;
+	}
+
+	TreeItem->SetDisplayName(InText);
+
+	if (TSharedPtr<FCineAssemblyFolderItem> FolderItem = TreeItem->AsFolder())
+	{
+		// If this is a folder item, update the path of all of its children (recursively)
+		SetChildrenPathRecursive(FolderItem, FolderItem->GetPath());
+		UpdateFolderList();
+	}
+
+	TreeItem->ParentFolder->Children.Sort(SortTreeItems);
+
+	TreeView->RequestTreeRefresh();
+}
+
+bool SCineAssemblyAssetTreeView::OnValidateTreeItemName(const FText& InText, FText& OutErrorMessage, TSharedPtr<ICineAssemblyTreeItem> TreeItem)
+{
+	if (InText.IsEmpty())
+	{
+		OutErrorMessage = LOCTEXT("EmptyItemNameErrorMessage", "Please provide a name for this item");
+		return false;
+	}
+
+	// Check for invalid characters
+	FString InvalidCharacters = INVALID_OBJECTNAME_CHARACTERS INVALID_LONGPACKAGE_CHARACTERS;
+
+	// These characters are actually valid, because we want to support naming tokens
+	InvalidCharacters = InvalidCharacters.Replace(TEXT("{}"), TEXT(""));
+	InvalidCharacters = InvalidCharacters.Replace(TEXT(":"), TEXT(""));
+
+	if (!FName::IsValidXName(InText.ToString(), InvalidCharacters, &OutErrorMessage))
+	{
+		return false;
+	}
+
+	// Check to see if the parent folder already contains an item of the same type with the same name
+	// Note: We only check folder names, because asset names may contain tokens, and identical tokenized names for different assets is fine
+	if (TreeItem->AsFolder())
+	{
+		const bool bIsDuplicateName = Algo::AnyOf(TreeItem->ParentFolder->GetChildFolders(), [TreeItem, InText](const TSharedPtr<ICineAssemblyTreeItem>& ChildItem)
+			{
+				return ChildItem && (ChildItem != TreeItem) && ChildItem->GetDisplayName().EqualTo(InText);
+			});
+
+		if (bIsDuplicateName)
+		{
+			OutErrorMessage = LOCTEXT("DuplicateItemNameErrorMessage", "An item already exists at this location with this name");
+			return false;
+		}
+	}
+
+	return true;
+}
+
+void SCineAssemblyAssetTreeView::OnTreeItemTextResolved(const FText& InText, TSharedPtr<ICineAssemblyTreeItem> TreeItem)
+{
+	if (TreeItem->AsSubAssembly())
+	{
+		if (UCineAssembly* Assembly = TreeItem->GetAssembly())
+		{
+			Assembly->AssemblyName.Resolved = InText;
+		}
+	}
+}
+
+void SCineAssemblyAssetTreeView::OnTreeViewDoubleClick(TSharedPtr<ICineAssemblyTreeItem> TreeItem)
+{
+	// Get the row for the input item and put its textblock in edit mode so the user can rename the item
+	if (TreeItem && TreeItem->NameWidget)
+	{
+		if (!bIsReadOnly && !TreeItem->IsReadOnly())
+		{
+			TreeItem->NameWidget->EnterEditingMode();
+		}
+	}
+}
+
+FReply SCineAssemblyAssetTreeView::OnTreeViewKeyDown(const FGeometry& MyGeometry, const FKeyEvent& InKeyEvent)
+{
+	if (bIsReadOnly)
+	{
+		return FReply::Unhandled();
+	}
+
+	if (InKeyEvent.GetKey() == EKeys::Delete)
+	{
+		// The TreeView uses single selection mode, so at most one item can ever be selected by the user
+		TArray<TSharedPtr<ICineAssemblyTreeItem>> SelectedItems = TreeView->GetSelectedItems();
+		if (SelectedItems.Num() == 1)
+		{
+			const TSharedPtr<ICineAssemblyTreeItem>& SelectedItem = SelectedItems[0];
+			if (!SelectedItem->IsReadOnly())
+			{
+				DeleteTreeItem(SelectedItem);
+				return FReply::Handled();
+			}
+		}
+	}
+
+	return FReply::Unhandled();
+}
+
+TSharedPtr<SWidget> SCineAssemblyAssetTreeView::MakeContentTreeContextMenu()
+{
+	// The TreeView uses single selection mode, so at most one item can ever be selected by the user
+	TArray<TSharedPtr<ICineAssemblyTreeItem>> SelectedItems = TreeView->GetSelectedItems();
+	if (SelectedItems.Num() == 1)
+	{
+		constexpr bool bCloseAfterMenuSelection = true;
+		FMenuBuilder MenuBuilder(bCloseAfterMenuSelection, nullptr);
+
+		TSharedPtr<ICineAssemblyTreeItem>& SelectedTreeItem = SelectedItems[0];
+
+		if (TSharedPtr<FCineAssemblyFolderItem> FolderItem = SelectedTreeItem->AsFolder())
+		{
+			MenuBuilder.AddMenuEntry(
+				LOCTEXT("AddChildFolderAction", "Add Child Folder"),
+				LOCTEXT("AddChildFolderTooltip", "Add Child Folder"),
+				FSlateIcon(FAppStyle::GetAppStyleSetName(), "Icons.Plus"),
+				FUIAction(FExecuteAction::CreateLambda([this]() { AddFolder(); })),
+				NAME_None,
+				EUserInterfaceActionType::Button
+			);
+		}
+
+		if (!bIsReadOnly && !SelectedTreeItem->IsReadOnly())
+		{
+			MenuBuilder.AddMenuEntry(
+				LOCTEXT("RenameAction", "Rename"),
+				LOCTEXT("RenameActionToolTip", "Rename"),
+				FSlateIcon(FCineAssemblyToolsStyle::StyleName, "Icons.AssetNaming"),
+				FUIAction(FExecuteAction::CreateLambda([SelectedTreeItem]() { SelectedTreeItem->NameWidget->EnterEditingMode(); })),
+				NAME_None,
+				EUserInterfaceActionType::Button
+			);
+
+			MenuBuilder.AddMenuEntry(
+				LOCTEXT("DeleteAction", "Delete"),
+				LOCTEXT("DeleteActionToolTip", "Delete"),
+				FSlateIcon(FAppStyle::GetAppStyleSetName(), "Icons.Delete"),
+				FUIAction(FExecuteAction::CreateSP(this, &SCineAssemblyAssetTreeView::DeleteTreeItem, SelectedTreeItem)),
+				NAME_None,
+				EUserInterfaceActionType::Button
+			);
+		}
+
+		return MenuBuilder.MakeWidget();
+	}
+
+	return SNullWidget::NullWidget;
+}
+
+void SCineAssemblyAssetTreeView::DeleteTreeItem(TSharedPtr<ICineAssemblyTreeItem> TreeItem)
+{
+	if (TSharedPtr<FCineAssemblyFolderItem> FolderItem = TreeItem->AsFolder())
+	{
+		DeleteFolderItem(FolderItem);
+	}
+	else if (TSharedPtr<FSubAssemblySectionItem> SectionItem = TreeItem->AsSubAssemblySection())
+	{
+		UE::CineAssemblyTools::Private::RemoveSubAssemblySection(SectionItem->SubAssemblySection.Get());
+	}
+
+	TreeItem->ParentFolder->Children.Remove(TreeItem);
+
+	UpdateFolderList();
+
+	TreeView->RequestTreeRefresh();
+
+	OnItemRemoved.ExecuteIfBound();
+}
+
+void SCineAssemblyAssetTreeView::DeleteFolderItem(const TSharedPtr<FCineAssemblyFolderItem>& FolderItem)
+{
+	// If the folder contains the top level assembly, reparent that item to the root so it is not lost
+	if (FolderItem->Children.Contains(TopLevelAssemblyItem))
+	{
+		FolderItem->Children.Remove(TopLevelAssemblyItem);
+		RootFolder->Children.Add(TopLevelAssemblyItem);
+		TopLevelAssemblyItem->ParentFolder = RootFolder;
+	}
+
+	for (const TSharedPtr<ICineAssemblyTreeItem>& Child : FolderItem->Children)
+	{
+		if (TSharedPtr<FCineAssemblyFolderItem> ChildFolder = Child->AsFolder())
+		{
+			DeleteFolderItem(ChildFolder);
+		}
+		else if (TSharedPtr<FSubAssemblySectionItem> ChildSection = Child->AsSubAssemblySection())
+		{
+			UE::CineAssemblyTools::Private::RemoveSubAssemblySection(ChildSection->SubAssemblySection.Get());
+		}
+	}
+
+	FolderItem->Children.Empty();
+}
+
+void SCineAssemblyAssetTreeView::UpdateFolderList()
+{
+	using namespace UE::CineAssemblyTools::Private;
+
+	UCineAssembly* Assembly = WeakAssembly.Get();
+	if (!Assembly)
+	{
+		return;
+	}
+
+	TArray<FString> CachedFolderList = GetFolderTemplates(Assembly);
+
+	TArray<FString> FolderList;
+	GetFolderListRecursive(RootFolder, FolderList);
+
+	if (FolderList != CachedFolderList)
+	{
+		Assembly->Modify();
+		SetFolderTemplates(Assembly, FolderList);
+	}
+
+	if (TopLevelAssemblyItem->ParentFolder)
+	{
+		Assembly->PathRelativeToRoot.Template = TopLevelAssemblyItem->ParentFolder->GetPath();
+	}
+}
+
+void SCineAssemblyAssetTreeView::ExpandTreeRecursive(const TSharedPtr<FCineAssemblyFolderItem>& FolderItem) const
+{
+	TreeView->SetItemExpansion(FolderItem, true);
+
+	for (const TSharedPtr<FCineAssemblyFolderItem>& ChildFolder : FolderItem->GetChildFolders())
+	{
+		ExpandTreeRecursive(ChildFolder);
+	}
+}
+
+FReply SCineAssemblyAssetTreeView::OnTreeRowDragDetected(const FGeometry& MyGeometry, const FPointerEvent& MouseEvent)
+{
+	if (bIsReadOnly)
+	{
+		return FReply::Unhandled();
+	}
+
+	if (MouseEvent.IsMouseButtonDown(EKeys::LeftMouseButton))
+	{
+		TArray<TSharedPtr<ICineAssemblyTreeItem>> SelectedItems = TreeView->GetSelectedItems();
+		if (SelectedItems.Num() == 1)
+		{
+			TSharedPtr<ICineAssemblyTreeItem> SelectedTreeItem = SelectedItems[0];
+
+			if (!SelectedTreeItem->CanDrag())
+			{
+				return FReply::Handled();
+			}
+
+			TSharedRef<FCineAssemblyTreeDragDrop> Operation = FCineAssemblyTreeDragDrop::New(SelectedTreeItem);
+
+			// If the drop operation is not handled for some reason, add the selected tree item back to its original location
+			Operation->OnDropNotHandled.BindLambda([this, SelectedTreeItem]()
+				{
+					SelectedTreeItem->ParentFolder->Children.Add(SelectedTreeItem);
+					SelectedTreeItem->ParentFolder->Children.Sort(UE::CineAssemblyTools::Private::SortTreeItems);
+
+					TreeView->RequestTreeRefresh();
+				});
+
+			SelectedTreeItem->ParentFolder->Children.Remove(SelectedTreeItem);
+			TreeView->RequestTreeRefresh();
+
+			return FReply::Handled().BeginDragDrop(Operation);
+		}
+	}
+
+	return FReply::Unhandled();
+}
+
+FReply SCineAssemblyAssetTreeView::OnTreeRowAcceptDrop(const FDragDropEvent& InDragDropEvent, EItemDropZone InDropZone, TSharedPtr<ICineAssemblyTreeItem> InItem)
+{
+	using namespace UE::CineAssemblyTools::Private;
+
+	// Only Folder Items are valid drop targets
+	TSharedPtr<FCineAssemblyFolderItem> TargetFolder = InItem->AsFolder();
+	if (!TargetFolder)
+	{
+		return FReply::Unhandled();
+	}
+
+	TSharedPtr<FCineAssemblyTreeDragDrop> Operation = InDragDropEvent.GetOperationAs<FCineAssemblyTreeDragDrop>();
+	if (Operation.IsValid() && Operation->TreeItem.IsValid())
+	{
+		TSharedPtr<ICineAssemblyTreeItem> DroppedItem = Operation->TreeItem;
+
+		// Ensure that we are not attempting to drop an item onto itself or one of its children
+		if (DroppedItem == InItem)
+		{
+			return FReply::Unhandled();
+		}
+		else if (DroppedItem->AsFolder() && ContainsRecursive(DroppedItem->AsFolder(), InItem))
+		{
+			return FReply::Unhandled();
+		}
+
+		// If the dropped item is a folder, but the target folder already contains a subfolder with the same name, do not handle the drop (the dragged item will be reset to its original location)
+		const FString DroppedItemPath = TargetFolder->GetPath() / DroppedItem->GetName();
+		if (DroppedItem->AsFolder())
+		{
+			if (Algo::AnyOf(TargetFolder->GetChildFolders(), [DroppedItemPath](const TSharedPtr<FCineAssemblyFolderItem>& FolderPtr) { return FolderPtr && FolderPtr->GetPath() == DroppedItemPath; }))
+			{
+				return FReply::Unhandled();
+			}
+		}
+
+		DroppedItem->SetPath(DroppedItemPath);
+		DroppedItem->ParentFolder = TargetFolder;
+
+		TargetFolder->Children.Add(DroppedItem);
+
+		if (TSharedPtr<FCineAssemblyFolderItem> DroppedFolder = DroppedItem->AsFolder())
+		{
+			SetChildrenPathRecursive(DroppedFolder, DroppedFolder->GetPath());
+		}
+
+		TargetFolder->Children.Sort(SortTreeItems);
+
+		UpdateFolderList();
+
+		ExpandTreeRecursive(TargetFolder);
+		TreeView->RequestTreeRefresh();
+	}
+
+	return FReply::Handled();
+}
+
+TSharedRef<FCineAssemblyTreeDragDrop> FCineAssemblyTreeDragDrop::New(const TSharedPtr<ICineAssemblyTreeItem>& InItem)
+{
+	TSharedRef<FCineAssemblyTreeDragDrop> DragDropOp = MakeShared<FCineAssemblyTreeDragDrop>();
+	DragDropOp->TreeItem = InItem;
+	DragDropOp->MouseCursor = EMouseCursor::GrabHandClosed;
+	DragDropOp->Construct();
+
+	return DragDropOp;
+}
+
+TSharedPtr<SWidget> FCineAssemblyTreeDragDrop::GetDefaultDecorator() const
+{
+	return SNew(SBorder)
+		.Padding(8.0f)
+		.BorderImage(FCineAssemblyToolsStyle::Get().GetBrush("ProductionWizard.PanelBackground"))
+		[
+			SNew(SHorizontalBox)
+
+				+ SHorizontalBox::Slot()
+				.AutoWidth()
+				.VAlign(VAlign_Center)
+				.Padding(0.0f, 0.0f, 8.0f, 0.0f)
+				[
+					SNew(SImage).Image(FCineAssemblyToolsStyle::Get().GetBrush("Icons.Sequencer"))
+				]
+
+				+ SHorizontalBox::Slot()
+				.FillWidth(1.0f)
+				.VAlign(VAlign_Center)
+				[
+					SNew(STextBlock).Text(TreeItem->GetDisplayName())
+				]
+		];
+}
+
+void FCineAssemblyTreeDragDrop::OnDragged(const class FDragDropEvent& DragDropEvent)
+{
+	if (CursorDecoratorWindow.IsValid())
+	{
+		CursorDecoratorWindow->MoveWindowTo(DragDropEvent.GetScreenSpacePosition() - (CursorDecoratorWindow->GetSizeInScreen() * FVector2f(0.0f, 0.5f)));
+	}
+}
+
+void FCineAssemblyTreeDragDrop::OnDrop(bool bDropWasHandled, const FPointerEvent& MouseEvent)
+{
+	if (!bDropWasHandled)
+	{
+		OnDropNotHandled.ExecuteIfBound();
+	}
+}
+
+#undef LOCTEXT_NAMESPACE
diff --git a/Engine/Plugins/VirtualProduction/CinematicAssemblyTools/Source/CineAssemblyToolsEditor/Private/UI/CineAssembly/SCineAssemblyAssetTreeView.h b/Engine/Plugins/VirtualProduction/CinematicAssemblyTools/Source/CineAssemblyToolsEditor/Private/UI/CineAssembly/SCineAssemblyAssetTreeView.h
new file mode 100644
--- /dev/null
+++ b/Engine/Plugins/VirtualProduction/CinematicAssemblyTools/Source/CineAssemblyToolsEditor/Private/UI/CineAssembly/SCineAssemblyAssetTreeView.h
@@ -0,0 +1,177 @@
+// Copyright Epic Games, Inc. All Rights Reserved.
+
+#pragma once
+
+#include "Widgets/SCompoundWidget.h"
+
+#include "CineAssembly.h"
+#include "Framework/Views/ITypedTableView.h"
+#include "SNamingTokensEditableTextBox.h"
+#include "Widgets/Views/STreeView.h"
+
+struct ICineAssemblyTreeItem;
+struct FSubAssemblyItem;
+struct FSubAssemblySectionItem;
+struct FCineAssemblyFolderItem;
+
+class UMovieSceneSubAssemblySection;
+
+/** Drag and Drop operation to handle dragging tree items in the CineAssembly AssetTreeView  */
+class FCineAssemblyTreeDragDrop : public FDragDropOperation
+{
+public:
+
+	DRAG_DROP_OPERATOR_TYPE(FCineAssemblyTreeDragDrop, FDragDropOperation)
+
+	FCineAssemblyTreeDragDrop() = default;
+
+	static TSharedRef<FCineAssemblyTreeDragDrop> New(const TSharedPtr<ICineAssemblyTreeItem>& InItem);
+
+	// Begin FDragDropOperation overrides
+	virtual void OnDrop(bool bDropWasHandled, const FPointerEvent& MouseEvent) override;
+	virtual void OnDragged(const class FDragDropEvent& DragDropEvent) override;
+	virtual TSharedPtr<SWidget> GetDefaultDecorator() const override;
+	// End FDragDropOperation overrides
+
+	/** The tree item being dragged. */
+	TSharedPtr<ICineAssemblyTreeItem> TreeItem;
+
+	/** Callback to handle the case where an item is dropped somewhere not on the tree view */
+	FSimpleDelegate OnDropNotHandled;
+};
+
+/**
+ * A tree view widget that displays the folder hierarchy and asset layout of a CineAssembly.
+ */
+class SCineAssemblyAssetTreeView : public SCompoundWidget
+{
+public:
+	SCineAssemblyAssetTreeView() = default;
+
+	SLATE_BEGIN_ARGS(SCineAssemblyAssetTreeView)
+		: _SelectionMode(ESelectionMode::Type::Single)
+		, _ShouldEvaluateTokens(false)
+		, _IsReadOnly(false)
+		, _DisplayHintText(false)
+		{}
+
+		/** Sets the selection mode for this tree view. Ignored when IsReadOnly is true (selection is forced off). */
+		SLATE_ARGUMENT(ESelectionMode::Type, SelectionMode)
+
+		/** Sets whether this tree view should evaluate and display resolved token strings, rather than raw templates */
+		SLATE_ARGUMENT(bool, ShouldEvaluateTokens)
+
+		/** When true, all editing interactions are disabled: renaming, deleting, drag/drop, and keyboard shortcuts */
+		SLATE_ARGUMENT(bool, IsReadOnly)
+
+		/** When true, displays additional context text next to tree items (e.g. which assembly created a sub-assembly) */
+		SLATE_ARGUMENT(bool, DisplayHintText)
+
+		/** Triggered when an item in the tree is removed */
+		SLATE_EVENT(FSimpleDelegate, OnItemRemoved)
+
+	SLATE_END_ARGS()
+
+	void Construct(const FArguments& InArgs, UCineAssembly* InAssembly);
+
+	/** Adds a new child folder under the currently selected folder (or the root if nothing is selected) */
+	void AddFolder();
+
+	/** Rebuilds the content tree from the current assembly state and refreshes the tree view */
+	void Reinitialize();
+
+private:
+	/** Build the content tree for the input Assembly, including its SubAssemblies and default folders */
+	void InitializeContentTree(UCineAssembly* Assembly, TSharedPtr<FCineAssemblyFolderItem>& ContentRoot, TSharedPtr<FSubAssemblyItem>& TopLevelAssembly);
+
+	/** Add SubAssemblySection items to the tree from the Assembly's MovieScene */
+	void AddSubAssemblySectionsToTree(UCineAssembly* Assembly, TSharedPtr<FCineAssemblyFolderItem>& ContentRoot);
+
+	/** Add SubAssembly items to the tree from the Assembly's SubAssembly list */
+	void AddSubAssembliesToTree(UCineAssembly* Assembly, TSharedPtr<FCineAssemblyFolderItem>& ContentRoot);
+
+	/** Generates the row widget for an entry in the tree view */
+	TSharedRef<ITableRow> OnGenerateTreeRow(TSharedPtr<ICineAssemblyTreeItem> TreeItem, const TSharedRef<STableViewBase>& OwnerTable);
+
+	/** Gets the children of the input tree view item to build additional tree rows */
+	void OnGetChildren(TSharedPtr<ICineAssemblyTreeItem> TreeItem, TArray<TSharedPtr<ICineAssemblyTreeItem>>& OutNodes);
+
+	/** Callback when the content tree view rebuilds itself, used to make the textbox in the most recently added row editable */
+	void OnTreeItemsRebuilt();
+
+	/** Callback to get the display name of the input tree item */
+	FText GetTreeItemDisplayName(TSharedPtr<ICineAssemblyTreeItem> TreeItem) const;
+
+	/** Callback to get the secondary hint text displayed next to the input tree item */
+	FText GetTreeItemHintText(TSharedPtr<ICineAssemblyTreeItem> TreeItem) const;
+
+	/** Renames the tree item, and updates the paths of all of its children */
+	void OnTreeItemTextCommitted(const FText& InText, ETextCommit::Type InCommitType, TSharedPtr<ICineAssemblyTreeItem> TreeItem);
+
+	/** Validates a tree item's name, checking for empty names, invalid characters, and duplicate folder names */
+	bool OnValidateTreeItemName(const FText& InText, FText& OutErrorMessage, TSharedPtr<ICineAssemblyTreeItem> TreeItem);
+
+	/** Callback when naming tokens in a tree item's text are evaluated, used to update the assembly's resolved name */
+	void OnTreeItemTextResolved(const FText& InText, TSharedPtr<ICineAssemblyTreeItem> TreeItem);
+
+	/** Handles double clicks on the tree view to rename the item */
+	void OnTreeViewDoubleClick(TSharedPtr<ICineAssemblyTreeItem> TreeItem);
+
+	/** Handles key presses on the tree view */
+	FReply OnTreeViewKeyDown(const FGeometry& MyGeometry, const FKeyEvent& InKeyEvent);
+
+	/** Create the context menu when the content tree view is right-clicked */
+	TSharedPtr<SWidget> MakeContentTreeContextMenu();
+
+	/** Remove the input item and its children from the tree view */
+	void DeleteTreeItem(TSharedPtr<ICineAssemblyTreeItem> TreeItem);
+
+	/** Removes a folder and all of its children from the tree view */
+	void DeleteFolderItem(const TSharedPtr<FCineAssemblyFolderItem>& FolderItem);
+
+	/** Update the list of folders to create from the current set in the tree view */
+	void UpdateFolderList();
+
+	/** Recursively expands every folder in the tree view */
+	void ExpandTreeRecursive(const TSharedPtr<FCineAssemblyFolderItem>& FolderItem) const;
+
+	/** Begins a drag and drop event to drag an item out of the content tree view */
+	FReply OnTreeRowDragDetected(const FGeometry& MyGeometry, const FPointerEvent& MouseEvent);
+
+	/** Callback when an item is dropped onto a folder in the content tree view to add it to the children of that row */
+	FReply OnTreeRowAcceptDrop(const FDragDropEvent& InDragDropEvent, EItemDropZone InDropZone, TSharedPtr<ICineAssemblyTreeItem> InItem);
+
+private:
+	/** The assembly that this tree view was constructed with */
+	TWeakObjectPtr<UCineAssembly> WeakAssembly;
+
+	/** The tree view of content folders and assets for the provided CineAssembly (and its SubAssemblies) */
+	TSharedPtr<STreeView<TSharedPtr<ICineAssemblyTreeItem>>> TreeView;
+
+	/** Items source for the tree view */
+	TArray<TSharedPtr<ICineAssemblyTreeItem>> TreeItems;
+
+	/** The root folder in the tree view */
+	TSharedPtr<FCineAssemblyFolderItem> RootFolder;
+
+	/** An item representing where the top-level assembly should be created */
+	TSharedPtr<FSubAssemblyItem> TopLevelAssemblyItem;
+
+	/** The most recently added tree item, used to allow the user to immediately rename the item after the tree view updates to show it */
+	TSharedPtr<ICineAssemblyTreeItem> MostRecentlyAddedItem;
+
+	/** Delegate broadcast when an item is removed from the tree */
+	FSimpleDelegate OnItemRemoved;
+
+	/** The selection mode for this tree view */
+	ESelectionMode::Type SelectionMode = ESelectionMode::Single;
+
+	/** Whether this tree view should evaluate and display resolved token strings */
+	bool bShouldEvaluateTokens = false;
+
+	/** When true, all editing interactions are disabled */
+	bool bIsReadOnly = false;
+
+	/** When true, displays additional context text next to tree items */
+	bool bDisplayHintText = false;
+};
diff --git a/Engine/Plugins/VirtualProduction/CinematicAssemblyTools/Source/CineAssemblyToolsEditor/Private/UI/CineAssembly/SCineAssemblyConfigPanel.cpp b/Engine/Plugins/VirtualProduction/CinematicAssemblyTools/Source/CineAssemblyToolsEditor/Private/UI/CineAssembly/SCineAssemblyConfigPanel.cpp
--- a/Engine/Plugins/VirtualProduction/CinematicAssemblyTools/Source/CineAssemblyToolsEditor/Private/UI/CineAssembly/SCineAssemblyConfigPanel.cpp
+++ b/Engine/Plugins/VirtualProduction/CinematicAssemblyTools/Source/CineAssemblyToolsEditor/Private/UI/CineAssembly/SCineAssemblyConfigPanel.cpp
@@ -149,7 +149,7 @@ void SCineAssemblyConfigPanel::Refresh()
 	DetailsView->ForceRefresh();
 
 	// Recreate the hierarchy tree items based on the selected schema
-	PopulateHierarchyTree();
+	TreeView->Reinitialize();
 }
 
 TSharedRef<SWidget> SCineAssemblyConfigPanel::MakeDetailsWidget()
@@ -268,21 +268,10 @@ const FSlateBrush* SCineAssemblyConfigPanel::GetSchemaThumbnail() const
 
 TSharedRef<SWidget> SCineAssemblyConfigPanel::MakeHierarchyWidget()
 {
-	HierarchyTreeView = SNew(STreeView<TSharedPtr<FHierarchyTreeItem>>)
-		.TreeItemsSource(&HierarchyTreeItems)
-		.SelectionMode(ESelectionMode::None)
-		.OnGenerateRow(this, &SCineAssemblyConfigPanel::OnGenerateTreeRow)
-		.OnGetChildren(this, &SCineAssemblyConfigPanel::OnGetChildren);
-
-	// Create the hierarchy tree root
-	RootItem = MakeShared<FHierarchyTreeItem>();
-	RootItem->Type = FHierarchyTreeItem::EItemType::Folder;
-
-	FTemplateString RootTemplate;
-	RootTemplate.Template = TEXT("");
-	RootItem->Path = RootTemplate;
-
-	PopulateHierarchyTree();
+	TreeView = SNew(SCineAssemblyAssetTreeView, CineAssemblyToConfigure)
+		.IsReadOnly(true)
+		.DisplayHintText(true)
+		.ShouldEvaluateTokens(true);
 
 	return SNew(SVerticalBox)
 
@@ -301,160 +290,11 @@ TSharedRef<SWidget> SCineAssemblyConfigPanel::MakeHierarchyWidget()
 				.BorderImage(FCineAssemblyToolsStyle::Get().GetBrush("ProductionWizard.RecessedBackground"))
 				.Padding(8.0f)
 				[
-					HierarchyTreeView.ToSharedRef()
-				]
-		];
-}
-
-void SCineAssemblyConfigPanel::PopulateHierarchyTree()
-{
-	HierarchyTreeItems.Empty();
-	HierarchyTreeItems.Add(RootItem);
-	RootItem->ChildAssets.Empty();
-	RootItem->ChildFolders.Empty();
-
-	auto AddItemsToTree = [this](TArray<FTemplateString>& ItemList, FHierarchyTreeItem::EItemType ItemType)
-		{
-			// Sort the list so that paths are added to the tree in the proper order
-			ItemList.Sort([](const FTemplateString& A, const FTemplateString& B) { return A.Template < B.Template; });
-
-			for (const FTemplateString& ItemName : ItemList)
-			{
-				const FString ParentPath = FPaths::GetPath(ItemName.Template);
-
-				// Walk the tree until we find an item whose path matches the parent path. The new tree item will be created as one of its children
-				TSharedPtr<FHierarchyTreeItem> ParentItem = FindItemAtPathRecursive(RootItem, ParentPath);
-				if (ParentItem)
-				{
-					TSharedPtr<FHierarchyTreeItem> NewItem = MakeShared<FHierarchyTreeItem>();
-					NewItem->Type = ItemType;
-					NewItem->Path = ItemName;
-
-					if (ItemType == FHierarchyTreeItem::EItemType::Folder)
-					{
-						ParentItem->ChildFolders.Add(NewItem);
-					}
-					else
-					{
-						ParentItem->ChildAssets.Add(NewItem);
-					}
-				}
-			}
-		};
-
-	if (const UCineAssemblySchema* Schema = CineAssemblyToConfigure->GetSchema())
-	{
-		AddItemsToTree(CineAssemblyToConfigure->DefaultFolderNames, FHierarchyTreeItem::EItemType::Folder);
-		AddItemsToTree(CineAssemblyToConfigure->SubAssemblyNames, FHierarchyTreeItem::EItemType::Asset);
-
-		// Add the top-level assembly tree item
-		TSharedPtr<FHierarchyTreeItem> ParentItem = FindItemAtPathRecursive(RootItem, CineAssemblyToConfigure->PathRelativeToRoot.Template);
-		if (ParentItem)
-		{
-			TSharedPtr<FHierarchyTreeItem> NewItem = MakeShared<FHierarchyTreeItem>();
-			NewItem->Type = FHierarchyTreeItem::EItemType::Asset;
-			NewItem->Path.Template = TEXT("{assembly}");
-
-			ParentItem->ChildAssets.Add(NewItem);
-		}
-	}
-
-	HierarchyTreeView->RequestTreeRefresh();
-	ExpandTreeRecursive(RootItem);
-}
-
-void SCineAssemblyConfigPanel::ExpandTreeRecursive(TSharedPtr<FHierarchyTreeItem> TreeItem) const
-{
-	if (HierarchyTreeView)
-	{
-		HierarchyTreeView->SetItemExpansion(TreeItem, true);
-	}
-
-	for (const TSharedPtr<FHierarchyTreeItem>& ChildItem : TreeItem->ChildFolders)
-	{
-		ExpandTreeRecursive(ChildItem);
-	}
-}
-
-TSharedPtr<FHierarchyTreeItem> SCineAssemblyConfigPanel::FindItemAtPathRecursive(TSharedPtr<FHierarchyTreeItem> TreeItem, const FString& Path) const
-{
-	if (TreeItem->Path.Template.Equals(Path))
-	{
-		return TreeItem;
-	}
-
-	for (const TSharedPtr<FHierarchyTreeItem>& Child : TreeItem->ChildFolders)
-	{
-		if (const TSharedPtr<FHierarchyTreeItem>& ItemAtPath = FindItemAtPathRecursive(Child, Path))
-		{
-			return ItemAtPath;
-		}
-	}
-
-	return nullptr;
-}
-
-TSharedRef<ITableRow> SCineAssemblyConfigPanel::OnGenerateTreeRow(TSharedPtr<FHierarchyTreeItem> TreeItem, const TSharedRef<STableViewBase>& OwnerTable)
-{
-	TSharedPtr<SImage> Icon = SNew(SImage);
-
-	if (TreeItem->Type == FHierarchyTreeItem::EItemType::Folder)
-	{
-		Icon->SetImage(FCineAssemblyToolsStyle::Get().GetBrush("Icons.Folder"));
-		Icon->SetColorAndOpacity(FAppStyle::Get().GetSlateColor("ContentBrowser.DefaultFolderColor"));
-	}
-	else
-	{
-		Icon->SetImage(FCineAssemblyToolsStyle::Get().GetBrush("Icons.Sequencer"));
-		Icon->SetColorAndOpacity(FLinearColor::White);
-	}
-
-	return SNew(STableRow<TSharedPtr<FHierarchyTreeItem>>, OwnerTable)
-		.Padding(FMargin(8.0f, 2.0f, 8.0f, 0.0f))
-		.Content()
-		[
-			SNew(SHorizontalBox)
-
-			+ SHorizontalBox::Slot()
-				.AutoWidth()
-				.VAlign(VAlign_Center)
-				.Padding(0.0f, 0.0f, 2.0f, 0.0f)
-				[
-					Icon.ToSharedRef()
-				]
-
-			+ SHorizontalBox::Slot()
-				.VAlign(VAlign_Center)
-				[
-					SNew(SNamingTokensEditableTextBox)
-						.Contexts({ NamingTokenContext.Get() })
-						.FilterArgs(FilterArgs)
-						.EvaluationFrequency(1.0f)
-						.IsReadOnly(true)
-						.DisplayBorderImage(false)
-						.DisplayTokenIcon(false)
-						.Text_Lambda([this, TreeItem]() -> FText
-							{
-								return (TreeItem == RootItem) ? LOCTEXT("RootPathName", "Root Folder") : FText::FromString(FPaths::GetPathLeaf(TreeItem->Path.Template));
-							})
-						.OnTokenizedTextEvaluated_Lambda([this, TreeItem](const FText& InText)
-							{
-								if (TreeItem != RootItem)
-								{
-									TreeItem->Path.Resolved = InText;
-								}
-							})
+					TreeView.ToSharedRef()
 				]
 		];
 }
 
-void SCineAssemblyConfigPanel::OnGetChildren(TSharedPtr<FHierarchyTreeItem> TreeItem, TArray<TSharedPtr<FHierarchyTreeItem>>& OutNodes)
-{
-	// Display all of the child assets first, followed by all of the child folders
-	OutNodes.Append(TreeItem->ChildAssets);
-	OutNodes.Append(TreeItem->ChildFolders);
-}
-
 TSharedRef<SWidget> SCineAssemblyConfigPanel::MakeNotesWidget()
 {
 	return SNew(SVerticalBox)
diff --git a/Engine/Plugins/VirtualProduction/CinematicAssemblyTools/Source/CineAssemblyToolsEditor/Private/UI/CineAssembly/SCineAssemblyConfigPanel.h b/Engine/Plugins/VirtualProduction/CinematicAssemblyTools/Source/CineAssemblyToolsEditor/Private/UI/CineAssembly/SCineAssemblyConfigPanel.h
--- a/Engine/Plugins/VirtualProduction/CinematicAssemblyTools/Source/CineAssemblyToolsEditor/Private/UI/CineAssembly/SCineAssemblyConfigPanel.h
+++ b/Engine/Plugins/VirtualProduction/CinematicAssemblyTools/Source/CineAssemblyToolsEditor/Private/UI/CineAssembly/SCineAssemblyConfigPanel.h
@@ -8,33 +8,11 @@
 #include "CineAssemblyNamingTokens.h"
 #include "IDetailsView.h"
 #include "NamingTokensEngineSubsystem.h"
+#include "SCineAssemblyAssetTreeView.h"
 #include "SNamingTokensEditableTextBox.h"
 #include "Widgets/Layout/SWidgetSwitcher.h"
 #include "Widgets/Views/STreeView.h"
 
-/** An entry in the hierarchy tree view */
-struct FHierarchyTreeItem
-{
-	/** The types of items that can be represented in this tree view */
-	enum class EItemType : uint8
-	{
-		Asset,
-		Folder
-	};
-
-	/** The type of this tree item */
-	EItemType Type;
-
-	/** The relative path of this tree item (possibly containing tokens) */
-	FTemplateString Path;
-
-	/** The children of this item in the tree that are Asset types */
-	TArray<TSharedPtr<FHierarchyTreeItem>> ChildAssets;
-
-	/** The children of this item in the tree that are Folder types */
-	TArray<TSharedPtr<FHierarchyTreeItem>> ChildFolders;
-};
-
 /**
  * A widget that wraps details and config steps for setting up a cine assembly asset
  */
@@ -67,21 +45,6 @@ private:
 	/** Creates the widget to display for the Notes tab */
 	TSharedRef<SWidget> MakeNotesWidget();
 
-	/** Populate the tree view items from the list of folders and assets specified by the selected schema */
-	void PopulateHierarchyTree();
-
-	/** Generates the row widget for an entry in the tree view */
-	TSharedRef<ITableRow> OnGenerateTreeRow(TSharedPtr<FHierarchyTreeItem> TreeItem, const TSharedRef<STableViewBase>& OwnerTable);
-
-	/** Gets the children of the input tree view item to build additional tree rows */
-	void OnGetChildren(TSharedPtr<FHierarchyTreeItem> TreeItem, TArray<TSharedPtr<FHierarchyTreeItem>>& OutNodes);
-
-	/** Returns the tree item whose path matches the input path */
-	TSharedPtr<FHierarchyTreeItem> FindItemAtPathRecursive(TSharedPtr<FHierarchyTreeItem> TreeItem, const FString& Path) const;
-
-	/** Recursively expands every item in the tree view */
-	void ExpandTreeRecursive(TSharedPtr<FHierarchyTreeItem> TreeItem) const;
-
 	/** Validates the user input text for the assembly name */
 	bool ValidateAssemblyName(const FText& InText, FText& OutErrorMessage) const;
 
@@ -101,14 +64,8 @@ private:
 	/** Details View displaying the reflected properties of the Cine Assembly being configured */
 	TSharedPtr<IDetailsView> DetailsView;
 
-	/** Items source for the tree view */
-	TArray<TSharedPtr<FHierarchyTreeItem>> HierarchyTreeItems;
-
-	/** The root item in the tree view */
-	TSharedPtr<FHierarchyTreeItem> RootItem;
-
-	/** A read-only tree view of folders and assets that will be created based on the selected schema */
-	TSharedPtr<STreeView<TSharedPtr<FHierarchyTreeItem>>> HierarchyTreeView;
+	/** A TreeView to display the list of other assets and folders that this Assembly will create */
+	TSharedPtr<SCineAssemblyAssetTreeView> TreeView;
 
 	/** Naming Token Context object, with a pointer to the assembly being configured */
 	TStrongObjectPtr<UCineAssemblyNamingTokensContext> NamingTokenContext;
diff --git a/Engine/Plugins/VirtualProduction/CinematicAssemblyTools/Source/CineAssemblyToolsEditor/Private/UI/CineAssembly/SCineAssemblySchemaWindow.cpp b/Engine/Plugins/VirtualProduction/CinematicAssemblyTools/Source/CineAssemblyToolsEditor/Private/UI/CineAssembly/SCineAssemblySchemaWindow.cpp
--- a/Engine/Plugins/VirtualProduction/CinematicAssemblyTools/Source/CineAssemblyToolsEditor/Private/UI/CineAssembly/SCineAssemblySchemaWindow.cpp
+++ b/Engine/Plugins/VirtualProduction/CinematicAssemblyTools/Source/CineAssemblyToolsEditor/Private/UI/CineAssembly/SCineAssemblySchemaWindow.cpp
@@ -9,16 +9,13 @@
 #include "CineAssemblySchemaFactory.h"
 #include "CineAssemblyToolsStyle.h"
 #include "Editor.h"
-#include "Framework/MultiBox/MultiBoxBuilder.h"
 #include "ISequencerModule.h"
 #include "IStructureDetailsView.h"
 #include "Misc/ConfigCacheIni.h"
-#include "Misc/MessageDialog.h"
 #include "Modules/ModuleManager.h"
 #include "MovieScene.h"
 #include "MovieSceneSubAssemblySection.h"
 #include "MovieSceneSubAssemblyTrack.h"
-#include "NamingTokensEngineSubsystem.h"
 #include "ProductionSettings.h"
 #include "PropertyCustomizationHelpers.h"
 #include "PropertyEditorDelegates.h"
@@ -42,152 +39,6 @@
 
 #define LOCTEXT_NAMESPACE "SCineAssemblySchemaWindow"
 
-struct ISchemaContentHierarchyTreeItem : TSharedFromThis<ISchemaContentHierarchyTreeItem>
-{
-	virtual ~ISchemaContentHierarchyTreeItem() {}
-
-	virtual TSharedPtr<FSchemaContentHierarchyFolder> AsFolder() { return nullptr; }
-	virtual TSharedPtr<FSchemaContentHierarchyAsset> AsAsset() { return nullptr; }
-
-	/** Whether this item can accept an item being dropped onto it */
-	virtual TOptional<EItemDropZone> CanAcceptDrop() const { return TOptional<EItemDropZone>(); }
-	
-	/** Return the name of the tree item */
-	virtual FString GetName() const = 0;
-
-	/** Set the name of the tree item */
-	virtual void SetName(const FString& InName) = 0;
-
-	/** Return the display name of the tree item */
-	virtual FText GetDisplayName() const = 0;
-
-	/** Set the display name of the tree item */
-	virtual void SetDisplayName(const FText& InDisplayName) = 0;
-
-	/** Return the path of the tree item (relative to the root folder) */
-	virtual FString GetPath() const = 0;
-
-	/** Set the path of the tree item (relative to the root folder) */
-	virtual void SetPath(const FString& InPath) = 0;
-
-	/** The parent folder of this item */
-	TSharedPtr<FSchemaContentHierarchyFolder> ParentFolder;
-
-	/** The text widget that displays this item's name and supports renaming */
-	TSharedPtr<SNamingTokensEditableTextBox> NameWidget;
-};
-
-struct FSchemaContentHierarchyAsset : ISchemaContentHierarchyTreeItem
-{
-	//~ Begin ISchemaContentHierarchyTreeItem Interface
-	virtual TSharedPtr<FSchemaContentHierarchyAsset> AsAsset() override { return SharedThis(this); }
-
-	virtual FString GetName() const override
-	{
-		return SubAssemblySection.Get() ? SubAssemblySection.Get()->GetSequenceName().ToString() : TEXT("");
-	}
-
-	virtual void SetName(const FString& InName) override
-	{
-		if (SubAssemblySection.IsValid())
-		{
-			SubAssemblySection->SetSequenceName(FText::FromString(InName));
-		}
-	}
-
-	virtual FText GetDisplayName() const override
-	{
-		return SubAssemblySection.Get() ? SubAssemblySection.Get()->GetSequenceName() : FText::GetEmpty();
-	}
-
-	virtual void SetDisplayName(const FText& InDisplayName) override
-	{
-		if (SubAssemblySection.IsValid())
-		{
-			SubAssemblySection->SetSequenceName(InDisplayName);
-		}
-	}
-
-	virtual FString GetPath() const override
-	{
-		return SubAssemblySection.Get() ? SubAssemblySection.Get()->GetSequencePath() : TEXT("");
-	}
-
-	virtual void SetPath(const FString& InPath) override
-	{
-		if (SubAssemblySection.IsValid())
-		{
-			SubAssemblySection->SetSequencePath(FPaths::GetPath(InPath));
-		}
-	}
-	//~ End ISchemaContentHierarchyTreeItem Interface
-
-	/** Returns the path name of the item's Template Object */
-	FSoftObjectPath GetTemplateObjectPath() const
-	{
-		if (SubAssemblySection.IsValid() && SubAssemblySection->GetAssemblyTemplate())
-		{
-			return SubAssemblySection->GetAssemblyTemplate()->GetPathName();
-		}
-		return FSoftObjectPath();
-	}
-
-	/** The MovieScene section represented in the content tree by this asset item */
-	TWeakObjectPtr<UMovieSceneSubAssemblySection> SubAssemblySection;
-};
-
-struct FSchemaContentHierarchyFolder : ISchemaContentHierarchyTreeItem
-{
-	//~ Begin ISchemaContentHierarchyTreeItem Interface
-	virtual TSharedPtr<FSchemaContentHierarchyFolder> AsFolder() override { return SharedThis(this); }
-	virtual TOptional<EItemDropZone> CanAcceptDrop() const override { return EItemDropZone::OntoItem; }
-
-	virtual FString GetName() const override
-	{
-		return FPaths::GetPathLeaf(FullPath);
-	}
-
-	virtual void SetName(const FString& InName) override
-	{
-		FullPath = FPaths::GetPath(FullPath) / InName;
-	}
-
-	virtual FText GetDisplayName() const override
-	{
-		return FText::FromString(GetName());
-	}
-
-	virtual void SetDisplayName(const FText& InDisplayName) override
-	{
-		FullPath = FPaths::GetPath(FullPath) / InDisplayName.ToString();
-	}
-
-	virtual FString GetPath() const override
-	{
-		return FullPath;
-	}
-
-	virtual void SetPath(const FString& InPath) override
-	{
-		FullPath = InPath;
-	}
-	//~ End ISchemaContentHierarchyTreeItem Interface
-
-	/** The fully qualified path of this item (relative to the root folder) */
-	FString FullPath;
-
-	/** The child assets of this folder item */
-	TArray<TSharedPtr<FSchemaContentHierarchyAsset>> ChildAssets;
-
-	/** The child folders of this folder item */
-	TArray<TSharedPtr<FSchemaContentHierarchyFolder>> ChildFolders;
-};
-
-namespace UE::CineAssemblyTools::Private
-{
-	auto SortTreeItems = [](const TSharedPtr<ISchemaContentHierarchyTreeItem>& A, const TSharedPtr<ISchemaContentHierarchyTreeItem>& B) { return A->GetPath() < B->GetPath(); };
-};
-
 void SCineAssemblySchemaWindow::Construct(const FArguments& InArgs, const FString& InCreateAssetPath)
 {
 	// Create a new transient CineAssemblySchema to configure in the UI.
@@ -241,9 +92,6 @@ TSharedRef<SWidget> SCineAssemblySchemaWindow::BuildUI()
 			];
 	}
 
-	// Initialize the content tree view with the current list of assets and folders saved in the schema
-	InitializeContentTree();
-
 	// Check the UI config settings to determine whether or not to display engine/plugin content by default in this window
 	UContentBrowserSettings* ContentBrowserSettings = GetMutableDefault<UContentBrowserSettings>();
 
@@ -693,17 +541,10 @@ TSharedRef<SWidget> SCineAssemblySchemaWindow::MakeMetadataTabContent()
 
 TSharedRef<SWidget> SCineAssemblySchemaWindow::MakeHierarchyTabContent()
 {
-	TreeView = SNew(STreeView<TSharedPtr<ISchemaContentHierarchyTreeItem>>)
-		.TreeItemsSource(&TreeItems)
+	TreeView = SNew(SCineAssemblyAssetTreeView, SchemaToConfigure->TemplateSequence)
 		.SelectionMode(ESelectionMode::Single)
-		.OnGenerateRow(this, &SCineAssemblySchemaWindow::OnGenerateTreeRow)
-		.OnGetChildren(this, &SCineAssemblySchemaWindow::OnGetChildren)
-		.OnItemsRebuilt(this, &SCineAssemblySchemaWindow::OnTreeItemsRebuilt)
-		.OnContextMenuOpening(this, &SCineAssemblySchemaWindow::MakeContentTreeContextMenu)
-		.OnMouseButtonDoubleClick(this, &SCineAssemblySchemaWindow::OnTreeViewDoubleClick)
-		.OnKeyDownHandler(this, &SCineAssemblySchemaWindow::OnTreeViewKeyDown);
-
-	ExpandTreeRecursive(RootFolder);
+		.ShouldEvaluateTokens(false)
+		.OnItemRemoved_Lambda([this]() { Sequencer->NotifyMovieSceneDataChanged(EMovieSceneDataChangeType::MovieSceneStructureItemRemoved); });
 
 	TSharedRef<SWidget> ContentHierarchyWidget = SNew(SVerticalBox)
 
@@ -814,604 +655,12 @@ FReply SCineAssemblySchemaWindow::OnAddAssetClicked()
 	return FReply::Handled();
 }
 
-FString SCineAssemblySchemaWindow::MakeUniqueFolderPath(TSharedPtr<FSchemaContentHierarchyFolder> FolderItem)
-{
-	// This implementation is based on a similar utility in AssetTools for creating a unique asset name.
-	const FString BaseName = FolderItem->GetPath() / TEXT("NewFolder");
-
-	// Find the index in the string of the last non-numeric character
-	int32 CharIndex = BaseName.Len() - 1;
-	while (CharIndex >= 0 && BaseName[CharIndex] >= TEXT('0') && BaseName[CharIndex] <= TEXT('9'))
-	{
-		--CharIndex;
-	}
-
-	// Trim the numeric characters off the end of the BaseName string, but remember the integer that was trimmed off to increment and append to the output
-	int32 IntSuffix = 1;
-	FString TrimmedBaseName = BaseName;
-	if (CharIndex >= 0 && CharIndex < BaseName.Len() - 1)
-	{
-		TrimmedBaseName = BaseName.Left(CharIndex + 1);
-
-		const FString TrailingInteger = BaseName.RightChop(CharIndex + 1);
-		IntSuffix = FCString::Atoi(*TrailingInteger) + 1;
-	}
-
-	FString WorkingName = TrimmedBaseName;
-
-	while (Algo::ContainsBy(FolderItem->ChildFolders, WorkingName, &FSchemaContentHierarchyFolder::FullPath))
-	{
-		WorkingName = FString::Printf(TEXT("%s%d"), *TrimmedBaseName, IntSuffix);
-		IntSuffix++;
-	}
-
-	return WorkingName;
-}
-
 FReply SCineAssemblySchemaWindow::OnAddFolderClicked()
 {
-	TSharedPtr<FSchemaContentHierarchyFolder> NewFolder = MakeShared<FSchemaContentHierarchyFolder>();
-
-	// Get the parent item for the new folder being added (this can be the root folder if no parent is currently selected)
-	// The TreeView uses single selection mode, so at most one item can ever be selected by the user
-	TSharedPtr<FSchemaContentHierarchyFolder> ParentFolder = RootFolder;
-
-	TArray<TSharedPtr<ISchemaContentHierarchyTreeItem>> SelectedItems = TreeView->GetSelectedItems();
-	if (SelectedItems.Num() == 1)
-	{
-		ParentFolder = SelectedItems[0]->AsFolder() ? SelectedItems[0]->AsFolder() : SelectedItems[0]->ParentFolder;
-	}
-
-	NewFolder->SetPath(MakeUniqueFolderPath(ParentFolder));
-	NewFolder->ParentFolder = ParentFolder;
-
-	ParentFolder->ChildFolders.Add(NewFolder);
-
-	// Sort the children alphabetically to maintain a good ordering with the new folder
-	ParentFolder->ChildFolders.Sort(UE::CineAssemblyTools::Private::SortTreeItems);
-
-	// Save a reference to this item so that when the tree is rebuilt, we can immediately start editing its name
-	MostRecentlyAddedItem = NewFolder;
-
-	UpdateFolderList();
-
-	TreeView->SetItemExpansion(ParentFolder, true);
-	TreeView->RequestTreeRefresh();
-
-	return FReply::Handled();
-}
-
-TSharedRef<ITableRow> SCineAssemblySchemaWindow::OnGenerateTreeRow(TSharedPtr<ISchemaContentHierarchyTreeItem> TreeItem, const TSharedRef<STableViewBase>& OwnerTable)
-{
-	TSharedPtr<SImage> Icon = SNew(SImage);
-
-	if (TreeItem->AsFolder())
-	{
-		Icon->SetImage(FCineAssemblyToolsStyle::Get().GetBrush("Icons.Folder"));
-		Icon->SetColorAndOpacity(FAppStyle::Get().GetSlateColor("ContentBrowser.DefaultFolderColor"));
-	}
-	else
-	{
-		Icon->SetImage(FCineAssemblyToolsStyle::Get().GetBrush("Icons.Sequencer"));
-		Icon->SetColorAndOpacity(FLinearColor::White);
-	}
-
-	const bool bIsRootItem = (TreeItem == RootFolder);
-
-	return SNew(STableRow<TSharedPtr<ISchemaContentHierarchyTreeItem>>, OwnerTable)
-		.ShowSelection(true)
-		.Padding(FMargin(8.0f, 0.0f, 8.0f, 0.0f))
-		.OnCanAcceptDrop_Lambda([](const FDragDropEvent& InDragDropEvent, EItemDropZone InDropZone, TSharedPtr<ISchemaContentHierarchyTreeItem> InItem) { return InItem->CanAcceptDrop(); })
-		.OnAcceptDrop(this, &SCineAssemblySchemaWindow::OnTreeRowAcceptDrop)
-		.OnDragDetected(this, &SCineAssemblySchemaWindow::OnTreeRowDragDetected)
-		.Content()
-		[
-			SNew(SHorizontalBox)
-
-			+ SHorizontalBox::Slot()
-				.AutoWidth()
-				.VAlign(VAlign_Center)
-				.Padding(0.0f, 0.0f, 2.0f, 0.0f)
-				[
-					Icon.ToSharedRef()
-				]
-
-			+ SHorizontalBox::Slot()
-				.VAlign(VAlign_Center)
-				[
-					SAssignNew(TreeItem->NameWidget, SNamingTokensEditableTextBox)
-						.SupportInlineEdit(true)
-						.ShouldEvaluateTokens(false)
-						.TextBoxPadding(FMargin(4, 2, 0, 2))
-						.IsReadOnly(bIsRootItem)
-						.DisplayTokenIcon(false)
-						.DisplayBorderImage(false)
-						.DisplayErrorMessage(!bIsRootItem)
-						.Text(this, &SCineAssemblySchemaWindow::GetTreeItemDisplayName, TreeItem)
-						.OnValidateTokenizedText_Lambda([this, TreeItem](const FText& InText, FText& OutErrorMessage) -> bool
-							{
-								return TreeItem->AsFolder() ? ValidateFolderName(InText, OutErrorMessage, TreeItem->AsFolder()) : ValidateAssetName(InText, OutErrorMessage);
-							})
-						.OnTextCommitted(this, &SCineAssemblySchemaWindow::OnTreeItemTextCommitted, TreeItem)
-				]
-
-			+ SHorizontalBox::Slot()
-				.VAlign(VAlign_Center)
-				[
-					SNew(SObjectPropertyEntryBox)
-						.AllowCreate(false)
-						.AllowClear(true)
-						.DisplayThumbnail(false)
-						.DisplayCompactSize(true)
-						.AllowedClass(ULevelSequence::StaticClass())
-						.ToolTipText(LOCTEXT("SchemaAssetTemplateSelectionTooltip", "(Optional) Select an existing level sequence to use as a template for this subsequence."))
-						.Visibility(TreeItem->AsAsset() && TreeItem != TopLevelAssemblyItem ? EVisibility::Visible : EVisibility::Collapsed)
-						.ObjectPath(this, &SCineAssemblySchemaWindow::GetTemplateObjectPath, TreeItem->AsAsset())
-						.OnObjectChanged(this, &SCineAssemblySchemaWindow::OnTreeItemTemplateChanged, TreeItem->AsAsset())
-				]
-		];
-}
-
-FText SCineAssemblySchemaWindow::GetTreeItemDisplayName(TSharedPtr<ISchemaContentHierarchyTreeItem> TreeItem) const
-{
-	if (TreeItem == RootFolder)
-	{
-		return LOCTEXT("RootPathName", "Root Folder");
-	}
-	else if (TreeItem == TopLevelAssemblyItem)
-	{
-		return FText::FromString(TEXT("{assembly}"));
-	}
-	else
-	{
-		return TreeItem->GetDisplayName();
-	}
-}
-
-FString SCineAssemblySchemaWindow::GetTemplateObjectPath(TSharedPtr<FSchemaContentHierarchyAsset> AssetItem) const
-{
-	if (!AssetItem || AssetItem == TopLevelAssemblyItem)
-	{
-		return FString();
-	}
-
-	return AssetItem->GetTemplateObjectPath().ToString();
-}
-
-void SCineAssemblySchemaWindow::OnGetChildren(TSharedPtr<ISchemaContentHierarchyTreeItem> TreeItem, TArray<TSharedPtr<ISchemaContentHierarchyTreeItem>>& OutNodes)
-{
-	// Display all of the child assets first, followed by all of the child folders
-	if (TSharedPtr<FSchemaContentHierarchyFolder> FolderItem = TreeItem->AsFolder())
-	{
-		OutNodes.Append(FolderItem->ChildAssets);
-		OutNodes.Append(FolderItem->ChildFolders);
-	}
-}
-
-void SCineAssemblySchemaWindow::OnTreeItemsRebuilt()
-{
-	// Upon regenerating the tree view, allow the user to immediately interact with the name widget of the newly added folder in order to rename it
-	if (MostRecentlyAddedItem && MostRecentlyAddedItem->NameWidget)
-	{
-		MostRecentlyAddedItem->NameWidget->EnterEditingMode();
-		MostRecentlyAddedItem.Reset();
-	}
-}
-
-TSharedPtr<SWidget> SCineAssemblySchemaWindow::MakeContentTreeContextMenu()
-{
-	// The TreeView uses single selection mode, so at most one item can ever be selected by the user
-	TArray<TSharedPtr<ISchemaContentHierarchyTreeItem>> SelectedItems = TreeView->GetSelectedItems();
-	if (SelectedItems.Num() == 1)
-	{
-		constexpr bool bCloseAfterMenuSelection = true;
-		FMenuBuilder MenuBuilder(bCloseAfterMenuSelection, nullptr);
-
-		TSharedPtr<ISchemaContentHierarchyTreeItem>& SelectedTreeItem = SelectedItems[0];
-
-		if (TSharedPtr<FSchemaContentHierarchyFolder> FolderItem = SelectedTreeItem->AsFolder())
-		{
-			MenuBuilder.AddMenuEntry(
-				LOCTEXT("AddChildFolderAction", "Add Child Folder"),
-				LOCTEXT("AddChildFolderTooltip", "Add Child Folder"),
-				FSlateIcon(FAppStyle::GetAppStyleSetName(), "Icons.Plus"),
-				FUIAction(FExecuteAction::CreateLambda([this]() { OnAddFolderClicked(); })),
-				NAME_None,
-				EUserInterfaceActionType::Button
-			);
-		}
-
-		if ((SelectedTreeItem != RootFolder) && (SelectedTreeItem != TopLevelAssemblyItem)) 
-		{
-			MenuBuilder.AddMenuEntry(
-				LOCTEXT("RenameAction", "Rename"),
-				LOCTEXT("RenameActionToolTip", "Rename"),
-				FSlateIcon(FCineAssemblyToolsStyle::StyleName, "Icons.AssetNaming"),
-				FUIAction(FExecuteAction::CreateLambda([SelectedTreeItem]() { SelectedTreeItem->NameWidget->EnterEditingMode(); })),
-				NAME_None,
-				EUserInterfaceActionType::Button
-			);
-
-			MenuBuilder.AddMenuEntry(
-				LOCTEXT("DeleteAction", "Delete"),
-				LOCTEXT("DeleteActionToolTip", "Delete"),
-				FSlateIcon(FAppStyle::GetAppStyleSetName(), "Icons.Delete"),
-				FUIAction(FExecuteAction::CreateSP(this, &SCineAssemblySchemaWindow::DeleteTreeItem, SelectedTreeItem)),
-				NAME_None,
-				EUserInterfaceActionType::Button
-			);
-		}
-
-		return MenuBuilder.MakeWidget();
-	}
-
-	return SNullWidget::NullWidget;
-}
-
-void SCineAssemblySchemaWindow::OnTreeViewDoubleClick(TSharedPtr<ISchemaContentHierarchyTreeItem> TreeItem)
-{
-	// Get the row for the input item and put its textblock in edit mode so the user can rename the item
-	if (TreeItem && TreeItem->NameWidget)
-	{
-		if (TreeItem != RootFolder && TreeItem != TopLevelAssemblyItem)
-		{
-			TreeItem->NameWidget->EnterEditingMode();
-		}
-	}
-}
-
-FReply SCineAssemblySchemaWindow::OnTreeViewKeyDown(const FGeometry& MyGeometry, const FKeyEvent& InKeyEvent)
-{
-	if (InKeyEvent.GetKey() == EKeys::Delete)
-	{
-		// The TreeView uses single selection mode, so at most one item can ever be selected by the user
-		TArray<TSharedPtr<ISchemaContentHierarchyTreeItem>> SelectedItems = TreeView->GetSelectedItems();
-		if (SelectedItems.Num() == 1)
-		{
-			const TSharedPtr<ISchemaContentHierarchyTreeItem>& SelectedItem = SelectedItems[0];
-			if ((SelectedItem != RootFolder) && (SelectedItem != TopLevelAssemblyItem))
-			{
-				DeleteTreeItem(SelectedItem);
-			}
-		}
-	}
-
+	TreeView->AddFolder();
 	return FReply::Handled();
 }
 
-void SCineAssemblySchemaWindow::InitializeContentTree()
-{
-	// Create the content tree root
-	RootFolder = MakeShared<FSchemaContentHierarchyFolder>();
-	RootFolder->SetPath(TEXT(""));
-
-	TreeItems.Add(RootFolder);
-
-	// Sort the list of folders so that paths are added to the tree in the proper order
-	TArray<FString> FolderTemplates = GetFolderTemplates();
-	FolderTemplates.Sort();
-
-	for (const FString& FolderPath : FolderTemplates)
-	{
-		const FString ParentPath = FPaths::GetPath(FolderPath);
-
-		// Walk the tree until we find an item whose path matches the parent path. The new tree item will be created as one of its children
-		TSharedPtr<FSchemaContentHierarchyFolder> ParentFolder = FindFolderByPathRecursive(RootFolder, ParentPath);
-		if (ParentFolder)
-		{
-			TSharedPtr<FSchemaContentHierarchyFolder> NewItem = MakeShared<FSchemaContentHierarchyFolder>();
-			NewItem->SetPath(FolderPath);
-			NewItem->ParentFolder = ParentFolder;
-
-			ParentFolder->ChildFolders.Add(NewItem);
-		}
-	}
-
-	UpdateContentHierarchyAssets();
-
-	// Add the top-level assembly node
-	const FString& ToplevelAssemblyPath = GetDefaultAssemblyPath();
-	TSharedPtr<FSchemaContentHierarchyFolder> ParentFolder = FindFolderByPathRecursive(RootFolder, ToplevelAssemblyPath);
-	if (ParentFolder)
-	{
-		TopLevelAssemblyItem = MakeShared<FSchemaContentHierarchyAsset>();
-		TopLevelAssemblyItem->ParentFolder = ParentFolder;
-
-		ParentFolder->ChildAssets.Add(TopLevelAssemblyItem);
-	}
-}
-
-TArray<FString> SCineAssemblySchemaWindow::GetFolderTemplates()
-{
-	if (!SchemaToConfigure || !SchemaToConfigure->TemplateSequence)
-	{
-		return TArray<FString>();
-	}
-
-	TArrayView<FTemplateString> DefaultFolderNames = SchemaToConfigure->TemplateSequence->DefaultFolderNames;
-
-	// Extract the template folder names from the template sequence's default folder list
-	TArray<FString> TemplateFolderNames;
-	TemplateFolderNames.Reserve(DefaultFolderNames.Num());
-
-	Algo::Transform(DefaultFolderNames, TemplateFolderNames, [](const FTemplateString& TemplateString) { return TemplateString.Template; });
-
-	return TemplateFolderNames;
-}
-
-void SCineAssemblySchemaWindow::SetFolderTemplates(TArray<FString> InFolderTemplates)
-{
-	if (!SchemaToConfigure || !SchemaToConfigure->TemplateSequence)
-	{
-		return;
-	}
-
-	TArray<FTemplateString>& DefaultFolderNames = SchemaToConfigure->TemplateSequence->DefaultFolderNames;
-
-	DefaultFolderNames.Reset(InFolderTemplates.Num());
-	Algo::Transform(InFolderTemplates, DefaultFolderNames, [](const FString& TemplateString) { return FTemplateString(TemplateString, FText::GetEmpty()); });
-}
-
-FString SCineAssemblySchemaWindow::GetDefaultAssemblyPath()
-{
-	if (!SchemaToConfigure || !SchemaToConfigure->TemplateSequence)
-	{
-		return FString();
-	}
-
-	return SchemaToConfigure->TemplateSequence->PathRelativeToRoot.Template;
-}
-
-void SCineAssemblySchemaWindow::SetDefaultAssemblyPath(const FString& InPath)
-{
-	if (!SchemaToConfigure || !SchemaToConfigure->TemplateSequence)
-	{
-		return;
-	}
-
-	SchemaToConfigure->TemplateSequence->PathRelativeToRoot.Template = InPath;
-}
-
-void SCineAssemblySchemaWindow::ExpandTreeRecursive(TSharedPtr<FSchemaContentHierarchyFolder> FolderItem) const
-{
-	TreeView->SetItemExpansion(FolderItem, true);
-
-	for (const TSharedPtr<FSchemaContentHierarchyFolder>& ChildFolder : FolderItem->ChildFolders)
-	{
-		ExpandTreeRecursive(ChildFolder);
-	}
-}
-
-void SCineAssemblySchemaWindow::GetFolderListRecursive(TSharedPtr<FSchemaContentHierarchyFolder> FolderItem, TArray<FString>& OutFolderList) const
-{
-	for (const TSharedPtr<FSchemaContentHierarchyFolder>& ChildFolder : FolderItem->ChildFolders)
-	{
-		OutFolderList.Add(ChildFolder->GetPath());
-		GetFolderListRecursive(ChildFolder, OutFolderList);
-	}
-}
-
-void SCineAssemblySchemaWindow::GetAssetListRecursive(TSharedPtr<FSchemaContentHierarchyFolder> FolderItem, TArray<TSharedPtr<FSchemaContentHierarchyAsset>>& OutAssetList) const
-{
-	OutAssetList.Append(FolderItem->ChildAssets);
-
-	for (const TSharedPtr<FSchemaContentHierarchyFolder>& ChildFolder : FolderItem->ChildFolders)
-	{
-		GetAssetListRecursive(ChildFolder, OutAssetList);
-	}
-}
-
-TSharedPtr<FSchemaContentHierarchyFolder> SCineAssemblySchemaWindow::FindFolderByPathRecursive(TSharedPtr<FSchemaContentHierarchyFolder> FolderItem, const FString& FolderPath) const
-{
-	if (FolderItem->GetPath().Equals(FolderPath))
-	{
-		return FolderItem;
-	}
-
-	for (const TSharedPtr<FSchemaContentHierarchyFolder>& ChildFolder : FolderItem->ChildFolders)
-	{
-		if (const TSharedPtr<FSchemaContentHierarchyFolder>& MatchingFolder = FindFolderByPathRecursive(ChildFolder, FolderPath))
-		{
-			return MatchingFolder;
-		}
-	}
-
-	return nullptr;
-}
-
-void SCineAssemblySchemaWindow::SetChildrenPathRecursive(TSharedPtr<FSchemaContentHierarchyFolder> FolderItem, const FString& NewPath)
-{
-	for (const TSharedPtr<FSchemaContentHierarchyAsset>& AssetItem : FolderItem->ChildAssets)
-	{
-		AssetItem->SetPath(NewPath / AssetItem->GetName());
-	}
-
-	for (const TSharedPtr<FSchemaContentHierarchyFolder>& ChildFolder : FolderItem->ChildFolders)
-	{
-		const FString NewChildPath = NewPath / ChildFolder->GetName();
-		ChildFolder->SetPath(NewChildPath);
-
-		SetChildrenPathRecursive(ChildFolder, NewChildPath);
-	}
-}
-
-void SCineAssemblySchemaWindow::OnTreeItemTextCommitted(const FText& InText, ETextCommit::Type InCommitType, TSharedPtr<ISchemaContentHierarchyTreeItem> TreeItem)
-{
-	// Early-out if the name has not actually changed
-	const FString OldName = TreeItem->GetName();
-	if (OldName.Equals(InText.ToString()))
-	{
-		return;
-	}
-
-	TreeItem->SetDisplayName(InText);
-
-	if (TSharedPtr<FSchemaContentHierarchyFolder> FolderItem = TreeItem->AsFolder())
-	{
-		// If this is a folder item, update the path of all of its children (recursively)
-		SetChildrenPathRecursive(FolderItem, FolderItem->GetPath());
-
-		FolderItem->ParentFolder->ChildFolders.Sort(UE::CineAssemblyTools::Private::SortTreeItems);
-		UpdateFolderList();
-	}
-	else
-	{
-		TreeItem->ParentFolder->ChildAssets.Sort(UE::CineAssemblyTools::Private::SortTreeItems);
-	}
-
-	TreeView->RequestTreeRefresh();
-}
-
-void SCineAssemblySchemaWindow::OnTreeItemTemplateChanged(const FAssetData& InAssetData, TSharedPtr<FSchemaContentHierarchyAsset> AssetItem)
-{
-	if (!AssetItem || AssetItem == TopLevelAssemblyItem)
-	{
-		return;
-	}
-
-	if (AssetItem->SubAssemblySection.IsValid())
-	{
-		AssetItem->SubAssemblySection->SetAssemblyTemplate(InAssetData.GetAsset());
-	}
-
-	TreeView->RequestTreeRefresh();
-}
-
-bool SCineAssemblySchemaWindow::ValidateAssetName(const FText& InText, FText& OutErrorMessage) const
-{
-	// An empty name is invalid
-	if (InText.IsEmpty())
-	{
-		OutErrorMessage = LOCTEXT("EmptyAssetNameErrorMessage", "Please provide a name for the asset");
-		return false;
-	}
-
-	// Ensure that the name does not contain any characters that would be invalid for an asset name
-	// This matches the validation that would happen if the user was renaming an asset in the content browser
-	FString InvalidCharacters = INVALID_OBJECTNAME_CHARACTERS INVALID_LONGPACKAGE_CHARACTERS;
-
-	// These characters are actually valid, because we want to support naming tokens
-	InvalidCharacters = InvalidCharacters.Replace(TEXT("{}"), TEXT(""));
-	InvalidCharacters = InvalidCharacters.Replace(TEXT(":"), TEXT(""));
-
-	return FName::IsValidXName(InText.ToString(), InvalidCharacters, &OutErrorMessage);
-}
-
-bool SCineAssemblySchemaWindow::ValidateFolderName(const FText& InText, FText& OutErrorMessage, TSharedPtr<FSchemaContentHierarchyFolder> FolderItem) const
-{
-	// Check for empty text string
-	if (InText.IsEmpty())
-	{
-		OutErrorMessage = LOCTEXT("EmptyFolderNameErrorMessage", "Please provide a name for this folder");
-		return false;
-	}
-
-	// Check for invalid characters
-	FString InvalidCharacters = INVALID_OBJECTNAME_CHARACTERS INVALID_LONGPACKAGE_CHARACTERS;
-
-	// These characters are actually valid, because we want to support naming tokens
-	InvalidCharacters = InvalidCharacters.Replace(TEXT("{}"), TEXT(""));
-	InvalidCharacters = InvalidCharacters.Replace(TEXT(":"), TEXT(""));
-
-	if (!FName::IsValidXName(InText.ToString(), InvalidCharacters, &OutErrorMessage))
-	{
-		return false;
-	}
-
-	// Check for duplicate folder names
-	const FString ParentPath = FPaths::GetPath(FolderItem->GetPath());
-	TSharedPtr<FSchemaContentHierarchyFolder> MatchingFolder = FindFolderByPathRecursive(RootFolder, ParentPath / InText.ToString());
-
-	if (MatchingFolder && MatchingFolder != FolderItem)
-	{
-		OutErrorMessage = LOCTEXT("DuplicateFolderErrorMessage", "A folder already exists at this location with this name");
-		return false;
-	}
-
-	return true;
-}
-
-void SCineAssemblySchemaWindow::DeleteTreeItem(TSharedPtr<ISchemaContentHierarchyTreeItem> TreeItem)
-{
-	if (TSharedPtr<FSchemaContentHierarchyFolder> FolderItem = TreeItem->AsFolder())
-	{
-		DeleteFolderItem(FolderItem);
-		TreeItem->ParentFolder->ChildFolders.Remove(FolderItem);
-	}
-	else if (TSharedPtr<FSchemaContentHierarchyAsset> AssetItem = TreeItem->AsAsset())
-	{
-		DeleteAssetItem(AssetItem);
-		TreeItem->ParentFolder->ChildAssets.Remove(AssetItem);
-	}
-
-	UpdateFolderList();
-
-	TreeView->RequestTreeRefresh();
-
-	Sequencer->NotifyMovieSceneDataChanged(EMovieSceneDataChangeType::MovieSceneStructureItemRemoved);
-}
-
-void SCineAssemblySchemaWindow::DeleteFolderItem(TSharedPtr<FSchemaContentHierarchyFolder> FolderItem)
-{
-	// If the folder contains the top level assembly, reparent that item to the root so it is not lost
-	if (FolderItem->ChildAssets.Contains(TopLevelAssemblyItem))
-	{
-		FolderItem->ChildAssets.Remove(TopLevelAssemblyItem);
-		RootFolder->ChildAssets.Add(TopLevelAssemblyItem);
-		TopLevelAssemblyItem->ParentFolder = RootFolder;
-	}
-
-	for (const TSharedPtr<FSchemaContentHierarchyAsset>& ChildAsset : FolderItem->ChildAssets)
-	{
-		DeleteAssetItem(ChildAsset);
-	}
-
-	FolderItem->ChildAssets.Empty();
-
-	for (const TSharedPtr<FSchemaContentHierarchyFolder>& ChildFolder : FolderItem->ChildFolders)
-	{
-		DeleteFolderItem(ChildFolder);
-	}
-	FolderItem->ChildFolders.Empty();
-}
-
-void SCineAssemblySchemaWindow::DeleteAssetItem(TSharedPtr<FSchemaContentHierarchyAsset> AssetItem)
-{
-	if (UMovieSceneSubAssemblySection* Section = AssetItem->SubAssemblySection.Get())
-	{
-		// Find the track that this section belongs to, and remove the section from that track. 
-		// IF the track has no other sections after deleting the section, then remove the whole track.
-		if (UMovieSceneSubAssemblyTrack* Track = Section->GetTypedOuter<UMovieSceneSubAssemblyTrack>())
-		{
-			Track->RemoveSection(*Section);
-			if (Track->GetAllSections().IsEmpty())
-			{
-				SchemaToConfigure->TemplateSequence->GetMovieScene()->RemoveTrack(*Track);
-			}
-		}
-	}
-}
-
-bool SCineAssemblySchemaWindow::ContainsRecursive(TSharedPtr<FSchemaContentHierarchyFolder> InParentFolder, TSharedPtr<ISchemaContentHierarchyTreeItem> InItem)
-{
-	if (InItem->AsAsset() && (InParentFolder->ChildAssets.Contains(InItem)))
-	{
-		return true;
-	}
-	else if (InItem->AsFolder() && (InParentFolder->ChildFolders.Contains(InItem)))
-	{
-		return true;
-	}
-
-	bool bFoundRecursively = false;
-	for (const TSharedPtr<FSchemaContentHierarchyFolder>& ChildFolder : InParentFolder->ChildFolders)
-	{
-		bFoundRecursively |= ContainsRecursive(ChildFolder, InItem);
-	}
-
-	return bFoundRecursively;
-}
-
 void SCineAssemblySchemaWindow::OnSchemaPropertiesChanged(const FPropertyChangedEvent& PropertyChangedEvent)
 {
 	const FName PropertyName = PropertyChangedEvent.GetPropertyName();
@@ -1457,22 +706,6 @@ FReply SCineAssemblySchemaWindow::OnCancelClicked()
 	return FReply::Handled();
 }
 
-void SCineAssemblySchemaWindow::UpdateFolderList()
-{
-	TArray<FString> CachedFolderList = GetFolderTemplates();
-
-	TArray<FString> FolderList;
-	GetFolderListRecursive(RootFolder, FolderList);
-
-	if (FolderList != CachedFolderList)
-	{
-		SchemaToConfigure->Modify();
-		SetFolderTemplates(FolderList);
-	}
-
-	SetDefaultAssemblyPath(TopLevelAssemblyItem->ParentFolder->GetPath());
-}
-
 void SCineAssemblySchemaWindow::OnMovieSceneDataChanged(EMovieSceneDataChangeType InChangeType)
 {
 	// In order to mute specific tracks, we only care about changes when a track is added or when a track value changes (which includes the user trying to unmute the track)
@@ -1497,72 +730,11 @@ void SCineAssemblySchemaWindow::OnMovieSceneDataChanged(EMovieSceneDataChangeTyp
 			}
 		}
 
-		// It is possible that a new Subsequence or Shot was just added or removed, so we update the subsequence items in the content tree
-		UpdateContentHierarchyAssets();
-
-		TreeView->RequestTreeRefresh();
-		ExpandTreeRecursive(RootFolder);
-	}
-}
-
-void SCineAssemblySchemaWindow::UpdateContentHierarchyAssets()
-{
-	// Cache the list of asset items that are already in the tree. Items who subsections are no longer in the template movie scene will be removed later.
-	TArray<TSharedPtr<FSchemaContentHierarchyAsset>> AssetItemsNotFound;
-	GetAssetListRecursive(RootFolder, AssetItemsNotFound);
-	AssetItemsNotFound.Remove(TopLevelAssemblyItem);
-
-	// Get the list of SubAssembly sections in the template movie scene
-	TArray<UMovieSceneSubAssemblySection*> SubAssemblySections;
-	if (UMovieScene* SchemaTemplateMovieScene = SchemaToConfigure->TemplateSequence->GetMovieScene())
-	{
-		for (UMovieSceneTrack* Track : SchemaTemplateMovieScene->GetTracks())
+		if (TreeView)
 		{
-			if (UMovieSceneSubAssemblyTrack* SubAssemblyTrack = Cast<UMovieSceneSubAssemblyTrack>(Track))
-			{
-				for (UMovieSceneSection* Section : SubAssemblyTrack->GetAllSections())
-				{
-					if (UMovieSceneSubAssemblySection* SubAssemblySection = Cast<UMovieSceneSubAssemblySection>(Section))
-					{
-						SubAssemblySections.Add(SubAssemblySection);
-					}
-				}
-			}
+			TreeView->Reinitialize();
 		}
 	}
-
-	for (UMovieSceneSubAssemblySection* SubAssemblySection : SubAssemblySections)
-	{
-		if (SubAssemblySection->IsTemplateSection())
-		{
-			// Walk the tree until we find a parent folder whose path matches the path of the SubAssembly section
-			if (TSharedPtr<FSchemaContentHierarchyFolder> ParentFolder = FindFolderByPathRecursive(RootFolder, SubAssemblySection->GetSequencePath()))
-			{
-				// Try to find a child asset item that already represents this SubAssembly section
-				TSharedPtr<FSchemaContentHierarchyAsset>* AssetItemPtr = ParentFolder->ChildAssets.FindByPredicate([SubAssemblySection](const TSharedPtr<FSchemaContentHierarchyAsset>& AssetItem) { return AssetItem->SubAssemblySection == SubAssemblySection; });
-				if (AssetItemPtr)
-				{
-					// If the asset item exists, then it is still valid (it matches one of the existing SubAssembly sections)
-					AssetItemsNotFound.Remove(*AssetItemPtr);
-				}
-				else
-				{
-					// If the asset item does not yet exist, make a new one and add it into the tree
-					TSharedPtr<FSchemaContentHierarchyAsset> NewItem = MakeShared<FSchemaContentHierarchyAsset>();
-					NewItem->ParentFolder = ParentFolder;
-					NewItem->SubAssemblySection = SubAssemblySection;
-					ParentFolder->ChildAssets.Add(NewItem);
-				}
-			}
-		}
-	}
-
-	// Any asset items that didn't match current SubAssembly sections need to be removed from the tree. 
-	// This could happen if a SubAssembly section is removed from the Schema Sequencer, for example.
-	for (const TSharedPtr<FSchemaContentHierarchyAsset>& AssetItem : AssetItemsNotFound)
-	{
-		AssetItem->ParentFolder->ChildAssets.Remove(AssetItem);
-	}
 }
 
 void SCineAssemblySchemaWindow::OnActorAddedToSequencer(AActor* InActor, const FGuid InBinding)
@@ -1643,175 +815,4 @@ EVisibility SCineAssemblySchemaWindow::GetDisplayRateWarningVisibility() const
 	return EVisibility::Collapsed;
 }
 
-FReply SCineAssemblySchemaWindow::OnTreeRowDragDetected(const FGeometry& MyGeometry, const FPointerEvent& MouseEvent)
-{
-	if (MouseEvent.IsMouseButtonDown(EKeys::LeftMouseButton))
-	{
-		TArray<TSharedPtr<ISchemaContentHierarchyTreeItem>> SelectedItems = TreeView->GetSelectedItems();
-		if (SelectedItems.Num() == 1)
-		{
-			TSharedPtr<ISchemaContentHierarchyTreeItem> SelectedTreeItem = SelectedItems[0];
-		
-			if (SelectedTreeItem == RootFolder)
-			{
-				return FReply::Handled();
-			}
-			
-			FSoftObjectPath TemplatePath = SelectedTreeItem->AsAsset() ? SelectedTreeItem->AsAsset()->GetTemplateObjectPath() : FSoftObjectPath();
-			TSharedRef<FSchemaAssetDragDrop> Operation = FSchemaAssetDragDrop::New(FSchemaAssetParameters(SelectedTreeItem->GetName(), TemplatePath));
-			Operation->SourceTreeItem = SelectedTreeItem;
-
-			if (TSharedPtr<FSchemaContentHierarchyAsset> AssetItem = SelectedTreeItem->AsAsset())
-			{
-				Operation->OnDropNotHandled.BindLambda([this, AssetItem](const FString& AssetName)
-					{
-						// If the tree item is not dropped in a valid place, restore it to its original place in the tree
-						AssetItem->ParentFolder->ChildAssets.Add(AssetItem);
-						TreeView->RequestTreeRefresh();
-					});
-
-				// Remove this item from the tree while it is being dragged
-				SelectedTreeItem->ParentFolder->ChildAssets.Remove(AssetItem);
-			}
-			else if (TSharedPtr<FSchemaContentHierarchyFolder> FolderItem = SelectedTreeItem->AsFolder())
-			{
-				Operation->OnDropNotHandled.BindLambda([this, FolderItem](const FString& AssetName)
-					{
-						// If the tree item is not dropped in a valid place, restore it to its original place in the tree
-						FolderItem->ParentFolder->ChildFolders.Add(FolderItem);
-						TreeView->RequestTreeRefresh();
-					});
-
-				// Remove this item from the tree while it is being dragged
-				FolderItem->ParentFolder->ChildFolders.Remove(FolderItem);
-			}
-
-			TreeView->RequestTreeRefresh();
-			return FReply::Handled().BeginDragDrop(Operation);
-		}
-	}
-
-	return FReply::Unhandled();
-}
-
-FReply SCineAssemblySchemaWindow::OnTreeRowAcceptDrop(const FDragDropEvent& InDragDropEvent, EItemDropZone InDropZone, TSharedPtr<ISchemaContentHierarchyTreeItem> InItem)
-{
-	// Only Folder Items are valid drop targets
-	TSharedPtr<FSchemaContentHierarchyFolder> TargetFolder = InItem->AsFolder();
-	if (!TargetFolder)
-	{
-		return FReply::Unhandled();
-	}
-
-	TSharedPtr<FSchemaAssetDragDrop> Operation = InDragDropEvent.GetOperationAs<FSchemaAssetDragDrop>();
-	if (Operation.IsValid())
-	{
-		TSharedPtr<ISchemaContentHierarchyTreeItem> DroppedItem;
-		if (Operation->SourceTreeItem)
-		{
-			DroppedItem = Operation->SourceTreeItem;
-		}
-		else
-		{
-			// This must be a new item from the asset list, so we need to make a new tree item for it
-			DroppedItem = MakeShared<FSchemaContentHierarchyAsset>();
-		}
-
-		// Ensure that we are not attempting to drop an item onto itself or one of its children
-		if (DroppedItem == InItem)
-		{
-			return FReply::Unhandled();
-		}
-		else if (DroppedItem->AsFolder() && ContainsRecursive(DroppedItem->AsFolder(), InItem))
-		{
-			return FReply::Unhandled();
-		}
-
-		// If the dropped item is a folder, but the target folder already contains a subfolder with the same name, do not handle the drop (the dragged item will be reset to its original location)
-		const FString DroppedItemPath = TargetFolder->GetPath() / Operation->ItemName;
-		if (DroppedItem->AsFolder())
-		{
-			if (Algo::AnyOf(TargetFolder->ChildFolders, [DroppedItemPath](const TSharedPtr<FSchemaContentHierarchyFolder>& FolderPtr) { return FolderPtr && FolderPtr->GetPath() == DroppedItemPath; }))
-			{
-				return FReply::Unhandled();
-			}
-		}
-
-		DroppedItem->SetPath(DroppedItemPath);
-		DroppedItem->ParentFolder = TargetFolder;
-
-		if (TSharedPtr<FSchemaContentHierarchyAsset> DroppedAsset = DroppedItem->AsAsset())
-		{
-			TargetFolder->ChildAssets.Add(DroppedAsset);
-			TargetFolder->ChildAssets.Sort(UE::CineAssemblyTools::Private::SortTreeItems);
-		}
-		else if (TSharedPtr<FSchemaContentHierarchyFolder> DroppedFolder = DroppedItem->AsFolder())
-		{
-			TargetFolder->ChildFolders.Add(DroppedFolder);
-
-			SetChildrenPathRecursive(DroppedFolder, DroppedFolder->GetPath());
-			TargetFolder->ChildFolders.Sort(UE::CineAssemblyTools::Private::SortTreeItems);
-		}
-
-		UpdateFolderList();
-
-		ExpandTreeRecursive(TargetFolder);
-		TreeView->RequestTreeRefresh();
-	}
-
-	return FReply::Handled();
-}
-
-TSharedRef<FSchemaAssetDragDrop> FSchemaAssetDragDrop::New(const FSchemaAssetParameters& InParameters)
-{
-	TSharedRef<FSchemaAssetDragDrop> DragDropOp = MakeShared<FSchemaAssetDragDrop>();
-	DragDropOp->ItemName = InParameters.Name;
-	DragDropOp->Template = InParameters.Template;
-	DragDropOp->MouseCursor = EMouseCursor::GrabHandClosed;
-	DragDropOp->Construct();
-
-	return DragDropOp;
-}
-
-TSharedPtr<SWidget> FSchemaAssetDragDrop::GetDefaultDecorator() const
-{
-	return SNew(SBorder)
-		.Padding(8.0f)
-		.BorderImage(FCineAssemblyToolsStyle::Get().GetBrush("ProductionWizard.PanelBackground"))
-		[
-			SNew(SHorizontalBox)
-
-			+ SHorizontalBox::Slot()
-				.AutoWidth()
-				.VAlign(VAlign_Center)
-				.Padding(0.0f, 0.0f, 8.0f, 0.0f)
-				[
-					SNew(SImage).Image(FCineAssemblyToolsStyle::Get().GetBrush("Icons.Sequencer"))
-				]
-
-			+ SHorizontalBox::Slot()
-				.FillWidth(1.0f)
-				.VAlign(VAlign_Center)
-				[
-					SNew(STextBlock).Text(FText::FromString(ItemName))
-				]
-		];
-}
-
-void FSchemaAssetDragDrop::OnDragged(const class FDragDropEvent& DragDropEvent)
-{
-	if (CursorDecoratorWindow.IsValid())
-	{
-		CursorDecoratorWindow->MoveWindowTo(DragDropEvent.GetScreenSpacePosition() - (CursorDecoratorWindow->GetSizeInScreen() * FVector2f(0.0f, 0.5f)));
-	}
-}
-
-void FSchemaAssetDragDrop::OnDrop(bool bDropWasHandled, const FPointerEvent& MouseEvent)
-{
-	if (!bDropWasHandled)
-	{
-		OnDropNotHandled.ExecuteIfBound(ItemName);
-	}
-}
-
 #undef LOCTEXT_NAMESPACE
diff --git a/Engine/Plugins/VirtualProduction/CinematicAssemblyTools/Source/CineAssemblyToolsEditor/Private/UI/CineAssembly/SCineAssemblySchemaWindow.h b/Engine/Plugins/VirtualProduction/CinematicAssemblyTools/Source/CineAssemblyToolsEditor/Private/UI/CineAssembly/SCineAssemblySchemaWindow.h
--- a/Engine/Plugins/VirtualProduction/CinematicAssemblyTools/Source/CineAssemblyToolsEditor/Private/UI/CineAssembly/SCineAssemblySchemaWindow.h
+++ b/Engine/Plugins/VirtualProduction/CinematicAssemblyTools/Source/CineAssemblyToolsEditor/Private/UI/CineAssembly/SCineAssemblySchemaWindow.h
@@ -5,18 +5,12 @@
 #include "Widgets/SCompoundWidget.h"
 
 #include "CineAssemblySchema.h"
-#include "Input/DragAndDrop.h"
 #include "ISequencer.h"
-#include "SNamingTokensEditableTextBox.h"
+#include "SCineAssemblyAssetTreeView.h"
 #include "Widgets/Layout/SWidgetSwitcher.h"
-#include "Widgets/Views/STreeView.h"
 
 struct FPropertyAndParent;
 
-struct ISchemaContentHierarchyTreeItem;
-struct FSchemaContentHierarchyFolder;
-struct FSchemaContentHierarchyAsset;
-
 /** Mode that configures the UI based on the intended user interactions with the Cine Assembly Schema asset */
 enum class ESchemaConfigMode : uint8
 {
@@ -24,48 +18,6 @@ enum class ESchemaConfigMode : uint8
 	Edit
 };
 
-/** Schema asset parameters describing a configured asset for the schema. */
-struct FSchemaAssetParameters
-{
-	/** Name of the asset. */
-	FString Name;
-
-	/** Template object to use to create the asset. */
-	FSoftObjectPath Template;
-};
-
-/** Drag and Drop operation to handle dragging schema assets and folders  */
-class FSchemaAssetDragDrop : public FDragDropOperation
-{
-public:
-
-	DRAG_DROP_OPERATOR_TYPE(FSchemaAssetDragDrop, FDragDropOperation)
-
-	DECLARE_DELEGATE_OneParam(FOnDropNotHanlded, const FString&);
-
-	FSchemaAssetDragDrop() = default;
-
-	static TSharedRef<FSchemaAssetDragDrop> New(const FSchemaAssetParameters& InParameters);
-
-	// Begin FDragDropOperation overrides
-	virtual void OnDrop(bool bDropWasHandled, const FPointerEvent& MouseEvent) override;
-	virtual void OnDragged(const class FDragDropEvent& DragDropEvent) override;
-	virtual TSharedPtr<SWidget> GetDefaultDecorator() const override;
-	// End FDragDropOperation overrides
-
-	/** The source item being dragged. */
-	TSharedPtr<ISchemaContentHierarchyTreeItem> SourceTreeItem;
-
-	/** The name of the item being displayed, used by the decorator and to name the new tree item that will get made when dropped */
-	FString ItemName;
-
-	/** Template object path. */
-	FSoftObjectPath Template;
-
-	/** Callback to handle the case where an item is dropped somewhere not on the tree view */
-	FOnDropNotHanlded OnDropNotHandled;
-};
-
 /**
  * A window for configuring the properties of a UCineAssemblySchema
  */
@@ -146,99 +98,9 @@ private:
 	/** Adds a new folder item to the content tree view */
 	FReply OnAddFolderClicked();
 
-	/** Generates the row widget for an entry in the tree view */
-	TSharedRef<ITableRow> OnGenerateTreeRow(TSharedPtr<ISchemaContentHierarchyTreeItem> TreeItem, const TSharedRef<STableViewBase>& OwnerTable);
-
-	/** Gets the children of the input tree view item to build additional tree rows */
-	void OnGetChildren(TSharedPtr<ISchemaContentHierarchyTreeItem> TreeItem, TArray<TSharedPtr<ISchemaContentHierarchyTreeItem>>& OutNodes);
-
-	/** Callback to get the display name of the input tree item */
-	FText GetTreeItemDisplayName(TSharedPtr<ISchemaContentHierarchyTreeItem> TreeItem) const;
-
-	/** Callback to get the Template object path of the input tree item */
-	FString GetTemplateObjectPath(TSharedPtr<FSchemaContentHierarchyAsset> AssetItem) const;
-
-	/** Callback when the content tree view rebuilds itself, used to make the textbox in the most recently added row editable */
-	void OnTreeItemsRebuilt();
-
-	/** Create the context menu when the content tree view is right-clicked */
-	TSharedPtr<SWidget> MakeContentTreeContextMenu();
-
-	/** Handles double clicks on the tree view to rename the item */
-	void OnTreeViewDoubleClick(TSharedPtr<ISchemaContentHierarchyTreeItem> TreeItem);
-
-	/** Handles key presses on the tree view */
-	FReply OnTreeViewKeyDown(const FGeometry& MyGeometry, const FKeyEvent& InKeyEvent);
-
-	/** Build the content tree from the list of folders and assets saved in the schema properties */
-	void InitializeContentTree();
-
-	/** Get the list of template folder names from the Schema's Template Sequence */
-	TArray<FString> GetFolderTemplates();
-
-	/** Set the list of template folder names in the Schema's Template Sequence */
-	void SetFolderTemplates(TArray<FString> InFolderTemplates);
-
-	/** Get the default path of the Schema's Template Sequence */
-	FString GetDefaultAssemblyPath();
-
-	/** Set the default path of the Schema's Template Sequence */
-	void SetDefaultAssemblyPath(const FString& InPath);
-
-	/** Recursively expands every folder in the tree view */
-	void ExpandTreeRecursive(TSharedPtr<FSchemaContentHierarchyFolder> FolderItem) const;
-
-	/** Recursively get the path of every folder item in the tree view */
-	void GetFolderListRecursive(TSharedPtr<FSchemaContentHierarchyFolder> FolderItem, TArray<FString>& OutFolderList) const;
-
-	/** Recursively find every asset item in the content hierarchy tree */
-	void GetAssetListRecursive(TSharedPtr<FSchemaContentHierarchyFolder> FolderItem, TArray<TSharedPtr<FSchemaContentHierarchyAsset>>& OutAssetList) const;
-
-	/** Returns the folder in the tree whose path matches the input path */
-	TSharedPtr<FSchemaContentHierarchyFolder> FindFolderByPathRecursive(TSharedPtr<FSchemaContentHierarchyFolder> FolderItem, const FString& FolderPath) const;
-
-	/** Modifies the path of the input folder item and all of its children (recursively) with the input new path */
-	void SetChildrenPathRecursive(TSharedPtr<FSchemaContentHierarchyFolder> FolderItem, const FString& NewPath);
-
-	/** Renames the tree item, and updates the paths of all of its children */
-	void OnTreeItemTextCommitted(const FText& InText, ETextCommit::Type InCommitType, TSharedPtr<ISchemaContentHierarchyTreeItem> TreeItem);
-
-	/** Updates the Template object of the given tree item. */
-	void OnTreeItemTemplateChanged(const FAssetData& InAssetData, TSharedPtr<FSchemaContentHierarchyAsset> AssetItem);
-
 	/** Callback when one of the properties of the schema being configured changes */
 	void OnSchemaPropertiesChanged(const FPropertyChangedEvent& PropertyChangedEvent);
 
-	/** Validate the text entered by the user to ensure it will be a valid asset name */
-	bool ValidateAssetName(const FText& InText, FText& OutErrorMessage) const;
-
-	/** Validate the text entered by the user to ensure it will be a valid folder name */
-	bool ValidateFolderName(const FText& InText, FText& OutErrorMessage, TSharedPtr<FSchemaContentHierarchyFolder> FolderItem) const;
-
-	/** Remove the input item and its children from the tree view */
-	void DeleteTreeItem(TSharedPtr<ISchemaContentHierarchyTreeItem> TreeItem);
-
-	/** Removes an asset item from the tree view, and from the template sequence */
-	void DeleteAssetItem(TSharedPtr<FSchemaContentHierarchyAsset> AssetItem);
-
-	/** Removes a folder and all of its children from the tree view */
-	void DeleteFolderItem(TSharedPtr<FSchemaContentHierarchyFolder> FolderItem);
-
-	/** Begins a drag and drop event to drag an item out of the content tree view */
-	FReply OnTreeRowDragDetected(const FGeometry& MyGeometry, const FPointerEvent& MouseEvent);
-
-	/** Callback when an item is dropped onto a folder in the content tree view to add it to the children of that row */
-	FReply OnTreeRowAcceptDrop(const FDragDropEvent& InDragDropEvent, EItemDropZone InDropZone, TSharedPtr<ISchemaContentHierarchyTreeItem> InItem);
-
-	/** Generates a unique folder path name, assuming the input folder will be the parent folder */
-	FString MakeUniqueFolderPath(TSharedPtr<FSchemaContentHierarchyFolder> FolderItem);
-
-	/** Returns true if the input parent folder, or any of its child folders, contains in the input item */
-	bool ContainsRecursive(TSharedPtr<FSchemaContentHierarchyFolder> InParentFolder, TSharedPtr<ISchemaContentHierarchyTreeItem> InItem);
-
-	/** Update the list of folders to create from the current set in the tree view */
-	void UpdateFolderList();
-
 	/** Callback when an actor is added to the template sequence, used to remove the actor(s) from the binding */
 	void OnActorAddedToSequencer(AActor* InActor, const FGuid InBinding);
 
@@ -274,21 +136,9 @@ private:
 	/** The root path where the configured schema will be created */
 	FString CreateAssetPath;
 
-	/** Items source for the tree view */
-	TArray<TSharedPtr<ISchemaContentHierarchyTreeItem>> TreeItems;
-
-	/** The root folder in the tree view */
-	TSharedPtr<FSchemaContentHierarchyFolder> RootFolder;
-
-	/** An item representing where the top-level assembly should be created */
-	TSharedPtr<FSchemaContentHierarchyAsset> TopLevelAssemblyItem;
-
-	/** The tree view of content folders and assets for this schema */
-	TSharedPtr<STreeView<TSharedPtr<ISchemaContentHierarchyTreeItem>>> TreeView;
-
-	/** The mostly recently added tree item, used to allow the user to immediately rename the item after the tree view updates to show it */
-	TSharedPtr<ISchemaContentHierarchyTreeItem> MostRecentlyAddedItem;
-
 	/** The sequencer used to display / edit the template sequence */
 	TSharedPtr<ISequencer> Sequencer;
+
+	/** A TreeView to display the list of other assets and folders that Assembly's made from this Schema will create */
+	TSharedPtr<SCineAssemblyAssetTreeView> TreeView;
 };
