ShotManagement: Add support to lock/unlock metadata editing in the Assembly Edit Properties window.
#rnx
#Synced-CL 51715153

diff --git a/Engine/Plugins/VirtualProduction/CinematicAssemblyTools/Source/CineAssemblyToolsEditor/Private/CineAssemblyCustomization.cpp b/Engine/Plugins/VirtualProduction/CinematicAssemblyTools/Source/CineAssemblyToolsEditor/Private/CineAssemblyCustomization.cpp
--- a/Engine/Plugins/VirtualProduction/CinematicAssemblyTools/Source/CineAssemblyToolsEditor/Private/CineAssemblyCustomization.cpp
+++ b/Engine/Plugins/VirtualProduction/CinematicAssemblyTools/Source/CineAssemblyToolsEditor/Private/CineAssemblyCustomization.cpp
@@ -15,11 +15,13 @@
 #include "PropertyCustomizationHelpers.h"
 #include "SNamingTokensEditableTextBox.h"
 #include "STemplateStringEditableTextBox.h"
+#include "Widgets/Images/SImage.h"
 #include "Widgets/Input/SCheckBox.h"
 #include "Widgets/Input/SComboButton.h"
 #include "Widgets/Input/SMultiLineEditableTextBox.h"
 #include "Widgets/Input/SNumericEntryBox.h"
 #include "Widgets/Layout/SBox.h"
+#include "Widgets/Layout/SWidgetSwitcher.h"
 #include "Widgets/Text/STextBlock.h"
 
 #define LOCTEXT_NAMESPACE "CineAssemblyCustomization"
@@ -129,183 +131,85 @@ void FCineAssemblyCustomization::CustomizeMetadataCategory(IDetailLayoutBuilder&
 		return;
 	}
 
-	// Metadata is editable for new (transient) assemblies, but read-only for existing assemblies, so we customize the details view differently
-	if (CustomizedCineAssembly->GetPackage() == GetTransientPackage())
-	{
-		// Add a new category for Schema Metadata properties
-		const FText MetadataCategoryName = FText::Format(LOCTEXT("SchemaMetadataCategoryName", "{0} Metadata"), FText::FromString(BaseSchema->SchemaName));
-		IDetailCategoryBuilder& MetadataCategory = DetailBuilder.EditCategory("SchemaMetadata", MetadataCategoryName);
+	const bool bIsBeingConfigured = CustomizedCineAssembly->GetPackage() == GetTransientPackage();
 
-		// Add a property row for each metadata struct in the customized Assembly's base schema
-		for (const FAssemblyMetadataDesc& MetadataDesc : BaseSchema->AssemblyMetadata)
-		{
-			if (MetadataDesc.Key.IsEmpty())
-			{
-				continue;
-			}
+	// Assemblies that are still being configured should always have editable metadata fields
+	// Previously created assemblies should default to having read-only metadata
+	bIsMetadataReadOnly = !bIsBeingConfigured;
 
-			CustomizedCineAssembly->AddMetadataNamingToken(MetadataDesc.Key);
+	// Add a new category for Schema Metadata properties
+	const FText MetadataCategoryName = FText::Format(LOCTEXT("SchemaMetadataCategoryName", "{0} Metadata"), FText::FromString(BaseSchema->SchemaName));
+	IDetailCategoryBuilder& MetadataCategory = DetailBuilder.EditCategory("SchemaMetadata", MetadataCategoryName);
 
-			TSharedPtr<SWidget> ValueWidget;
-			if (MetadataDesc.Type == ECineAssemblyMetadataType::String)
-			{
-				ValueWidget = SNew(SBox)
-					.MaxDesiredHeight(120.0f)
-					[
-						MakeStringValueWidget(MetadataDesc)
-					];
-			}
-			else if (MetadataDesc.Type == ECineAssemblyMetadataType::Bool)
-			{
-				ValueWidget = SNew(SCheckBox)
-					.IsChecked_Lambda([this, MetadataDesc]()
-						{
-							bool Value;
-							if (!CustomizedCineAssembly->GetMetadataAsBool(MetadataDesc.Key, Value))
-							{
-								Value = MetadataDesc.DefaultValue.Get<bool>();
-							}
-							return Value ? ECheckBoxState::Checked : ECheckBoxState::Unchecked;
-						})
-					.OnCheckStateChanged_Lambda([this, MetadataDesc](ECheckBoxState CheckBoxState)
-						{
-							const bool Value = (CheckBoxState == ECheckBoxState::Checked);
-							CustomizedCineAssembly->SetMetadataAsBool(MetadataDesc.Key, Value);
-						});
-			}
-			else if (MetadataDesc.Type == ECineAssemblyMetadataType::Integer)
-			{
-				ValueWidget = SNew(SNumericEntryBox<int32>)
-					.Font(FAppStyle::GetFontStyle(TEXT("PropertyWindow.NormalFont")))
-					.TypeInterface(LeadingZeroTypeInterface)
-					.Value_Lambda([this, MetadataDesc]()
-						{
-							int32 Value;
-							if (!CustomizedCineAssembly->GetMetadataAsInteger(MetadataDesc.Key, Value))
-							{
-								Value = MetadataDesc.DefaultValue.Get<int32>();
-							}
-							return Value;
-						})
-					.OnValueChanged_Lambda([this, MetadataDesc](int32 InValue)
-						{
-							// In order to preserve the desired number of leading zeroes, the integer value is converted to a string using the type interface
-							// and then stored in the assembly as string metadata.
-							const FString ValueString = LeadingZeroTypeInterface->ToString(InValue);
-							CustomizedCineAssembly->SetMetadataAsString(MetadataDesc.Key, ValueString);
-						});
-			}
-			else if (MetadataDesc.Type == ECineAssemblyMetadataType::Float)
-			{
-				ValueWidget = SNew(SNumericEntryBox<float>)
-					.Font(FAppStyle::GetFontStyle(TEXT("PropertyWindow.NormalFont")))
-					.Value_Lambda([this, MetadataDesc]()
-						{
-							float Value;
-							if (!CustomizedCineAssembly->GetMetadataAsFloat(MetadataDesc.Key, Value))
-							{
-								Value = MetadataDesc.DefaultValue.Get<float>();
-							}
-							return Value;
-						})
-					.OnValueChanged_Lambda([this, MetadataDesc](float InValue)
-						{
-							CustomizedCineAssembly->SetMetadataAsFloat(MetadataDesc.Key, InValue);
-						});
-			}
-			else if (MetadataDesc.Type == ECineAssemblyMetadataType::AssetPath)
-			{
-				ValueWidget = SNew(SObjectPropertyEntryBox)
-					.AllowedClass(MetadataDesc.AssetClass.ResolveClass())
-					.ThumbnailPool(DetailBuilder.GetThumbnailPool())
-					.AllowCreate(true)
-					.ObjectPath_Lambda([this, MetadataDesc]()
-						{
-							FString Value;
-							if (!CustomizedCineAssembly->GetMetadataAsString(MetadataDesc.Key, Value))
-							{
-								Value = MetadataDesc.DefaultValue.Get<FString>();
-							}
-							return Value;
-						})
-					.OnObjectChanged_Lambda([this, MetadataDesc](const FAssetData& InAssetData)
-						{
-							// Store the path of the selected object as a string
-							CustomizedCineAssembly->SetMetadataAsString(MetadataDesc.Key, InAssetData.GetObjectPathString());
-						});
-			}
-			else if (MetadataDesc.Type == ECineAssemblyMetadataType::CineAssembly)
-			{
-				ValueWidget = SNew(SObjectPropertyEntryBox)
-					.AllowedClass(UCineAssembly::StaticClass())
-					.ThumbnailPool(DetailBuilder.GetThumbnailPool())
-					.AllowCreate(true)
-					.OnShouldFilterAsset_Raw(this, &FCineAssemblyCustomization::ShouldFilterAssetBySchema, MetadataDesc.SchemaType)
-					.ObjectPath_Lambda([this, MetadataDesc]()
-						{
-							FString Value;
-							if (!CustomizedCineAssembly->GetMetadataAsString(MetadataDesc.Key, Value))
-							{
-								Value = MetadataDesc.DefaultValue.Get<FString>();
-							}
-							return Value;
-						})
-					.OnObjectChanged_Lambda([this, MetadataDesc](const FAssetData& InAssetData)
-						{
-							// Store the path of the selected object as a string
-							CustomizedCineAssembly->SetMetadataAsString(MetadataDesc.Key, InAssetData.GetObjectPathString());
-						});
-			}
-			else
-			{
-				checkNoEntry();
-			}
+	// Add a lock/unlock checkbox to the category header to determine whether the metadata values are read-only or editable
+	// Note: The checkbox will not be visible on Assemblies that are still being configured, so that their metadata is always editable
+	MetadataCategory.HeaderContent(
+		SNew(SHorizontalBox)
 
-			MetadataCategory.AddCustomRow(FText::FromString(MetadataDesc.Key))
-				.NameContent()
+		+ SHorizontalBox::Slot()
+		.FillWidth(1.0f)
+		.HAlign(HAlign_Right)
+		.VAlign(VAlign_Center)
+		[
+			SNew(SCheckBox)
+				.Style(&FAppStyle::Get().GetWidgetStyle<FCheckBoxStyle>("ToggleButtonCheckbox"))
+				.Visibility_Lambda([bIsBeingConfigured]() { return bIsBeingConfigured ? EVisibility::Collapsed : EVisibility::Visible; })
+				.IsChecked_Lambda([this]() { return bIsMetadataReadOnly ? ECheckBoxState::Unchecked : ECheckBoxState::Checked; })
+				.OnCheckStateChanged_Lambda([this](ECheckBoxState NewState) { bIsMetadataReadOnly = NewState != ECheckBoxState::Checked; })
+				.ToolTipText_Lambda([this]() { return bIsMetadataReadOnly ? LOCTEXT("UnlockMetadataTooltip", "Unlock metadata for editing") : LOCTEXT("LockMetadataTooltip", "Lock metadata (make read-only)"); })
 				[
-					SNew(STextBlock)
-						.Text(FText::FromString(MetadataDesc.Key))
-						.Font(IDetailLayoutBuilder::GetDetailFont())
+					SNew(SImage)
+						.DesiredSizeOverride(FVector2D(16, 16))
+						.ColorAndOpacity(FSlateColor::UseForeground())
+						.Image_Lambda([this]() { return bIsMetadataReadOnly ? FAppStyle::Get().GetBrush("PropertyWindow.Locked") : FAppStyle::Get().GetBrush("PropertyWindow.Unlocked"); })
 				]
-				.ValueContent()
-				.HAlign(HAlign_Fill)
-				[
-					ValueWidget.ToSharedRef()
-				];
-		}
-	}
-	else
-	{
-		const FText MetadataCategoryName = LOCTEXT("AllMetadataCategoryName", "Metadata");
-		IDetailCategoryBuilder& MetadataCategory = DetailBuilder.EditCategory("SchemaMetadata", MetadataCategoryName);
+		]
+	);
 
-		for (const FString& MetadataKey : CustomizedCineAssembly->GetMetadataKeys())
+	// Add a row for each metadata field defined in the schema
+	for (const FAssemblyMetadataDesc& MetadataDesc : BaseSchema->AssemblyMetadata)
+	{
+		if (MetadataDesc.Key.IsEmpty())
 		{
-			if (MetadataKey.IsEmpty())
-			{
-				continue;
-			}
+			continue;
+		}
 
-			FString Value;
-			if (CustomizedCineAssembly->GetMetadataAsString(MetadataKey, Value))
-			{
-				MetadataCategory.AddCustomRow(FText::FromString(MetadataKey))
-					.NameContent()
+		CustomizedCineAssembly->AddMetadataNamingToken(MetadataDesc.Key);
+
+		MetadataCategory.AddCustomRow(FText::FromString(MetadataDesc.Key))
+			.NameContent()
+			[
+				SNew(STextBlock)
+					.Text(FText::FromString(MetadataDesc.Key))
+					.Font(IDetailLayoutBuilder::GetDetailFont())
+			]
+			.ValueContent()
+			.HAlign(HAlign_Fill)
+			[
+				SNew(SWidgetSwitcher)
+					.WidgetIndex_Lambda([this]() 
+						{
+							constexpr int32 ReadOnlyWidgetIndex = 0;
+							constexpr int32 ReadWriteWidgetIndex = 1;
+							return bIsMetadataReadOnly ? ReadOnlyWidgetIndex : ReadWriteWidgetIndex;
+						})
+
+					+ SWidgetSwitcher::Slot()
 					[
 						SNew(STextBlock)
-							.Text(FText::FromString(MetadataKey))
 							.Font(IDetailLayoutBuilder::GetDetailFont())
+							.ColorAndOpacity(FSlateColor::UseSubduedForeground())
+							.Text_Lambda([this, MetadataDesc]()
+								{
+									FString Value;
+									CustomizedCineAssembly->GetMetadataAsString(MetadataDesc.Key, Value);
+									return FText::FromString(Value);
+								})
 					]
-					.ValueContent()
-					.HAlign(HAlign_Fill)
-					[
-						SNew(STextBlock)
-							.Text(FText::FromString(Value))
-							.Font(IDetailLayoutBuilder::GetDetailFont())
-					];
-			}
-		}
+
+					+ SWidgetSwitcher::Slot()
+					[MakeMetadataWidget(DetailBuilder, MetadataDesc)]
+			];
 	}
 }
 
@@ -358,6 +262,112 @@ void FCineAssemblyCustomization::CustomizeSubsequenceCategory(IDetailLayoutBuild
 	}
 }
 
+TSharedRef<SWidget> FCineAssemblyCustomization::MakeMetadataWidget(IDetailLayoutBuilder& DetailBuilder, const FAssemblyMetadataDesc& MetadataDesc)
+{
+	switch (MetadataDesc.Type)
+	{
+		case ECineAssemblyMetadataType::String:
+			return SNew(SBox)
+				.MaxDesiredHeight(120.0f)
+				[
+					MakeStringValueWidget(MetadataDesc)
+				];
+		case ECineAssemblyMetadataType::Bool:
+			return SNew(SCheckBox)
+				.IsChecked_Lambda([this, MetadataDesc]()
+					{
+						bool Value;
+						if (!CustomizedCineAssembly->GetMetadataAsBool(MetadataDesc.Key, Value))
+						{
+							Value = MetadataDesc.DefaultValue.Get<bool>();
+						}
+						return Value ? ECheckBoxState::Checked : ECheckBoxState::Unchecked;
+					})
+				.OnCheckStateChanged_Lambda([this, MetadataDesc](ECheckBoxState CheckBoxState)
+					{
+						const bool Value = (CheckBoxState == ECheckBoxState::Checked);
+						CustomizedCineAssembly->SetMetadataAsBool(MetadataDesc.Key, Value);
+					});
+		case ECineAssemblyMetadataType::Integer:
+			return SNew(SNumericEntryBox<int32>)
+				.Font(FAppStyle::GetFontStyle(TEXT("PropertyWindow.NormalFont")))
+				.TypeInterface(LeadingZeroTypeInterface)
+				.Value_Lambda([this, MetadataDesc]()
+					{
+						int32 Value;
+						if (!CustomizedCineAssembly->GetMetadataAsInteger(MetadataDesc.Key, Value))
+						{
+							Value = MetadataDesc.DefaultValue.Get<int32>();
+						}
+						return Value;
+					})
+				.OnValueChanged_Lambda([this, MetadataDesc](int32 InValue)
+					{
+						// In order to preserve the desired number of leading zeroes, the integer value is converted to a string using the type interface
+						// and then stored in the assembly as string metadata.
+						const FString ValueString = LeadingZeroTypeInterface->ToString(InValue);
+						CustomizedCineAssembly->SetMetadataAsString(MetadataDesc.Key, ValueString);
+					});
+		case ECineAssemblyMetadataType::Float:
+			return SNew(SNumericEntryBox<float>)
+				.Font(FAppStyle::GetFontStyle(TEXT("PropertyWindow.NormalFont")))
+				.Value_Lambda([this, MetadataDesc]()
+					{
+						float Value;
+						if (!CustomizedCineAssembly->GetMetadataAsFloat(MetadataDesc.Key, Value))
+						{
+							Value = MetadataDesc.DefaultValue.Get<float>();
+						}
+						return Value;
+					})
+				.OnValueChanged_Lambda([this, MetadataDesc](float InValue)
+					{
+						CustomizedCineAssembly->SetMetadataAsFloat(MetadataDesc.Key, InValue);
+					});
+		case ECineAssemblyMetadataType::AssetPath:
+			return SNew(SObjectPropertyEntryBox)
+				.AllowedClass(MetadataDesc.AssetClass.ResolveClass())
+				.ThumbnailPool(DetailBuilder.GetThumbnailPool())
+				.AllowCreate(true)
+				.ObjectPath_Lambda([this, MetadataDesc]()
+					{
+						FString Value;
+						if (!CustomizedCineAssembly->GetMetadataAsString(MetadataDesc.Key, Value))
+						{
+							Value = MetadataDesc.DefaultValue.Get<FString>();
+						}
+						return Value;
+					})
+				.OnObjectChanged_Lambda([this, MetadataDesc](const FAssetData& InAssetData)
+					{
+						CustomizedCineAssembly->SetMetadataAsString(MetadataDesc.Key, InAssetData.GetObjectPathString());
+					});
+		case ECineAssemblyMetadataType::CineAssembly:
+			return SNew(SObjectPropertyEntryBox)
+				.AllowedClass(UCineAssembly::StaticClass())
+				.ThumbnailPool(DetailBuilder.GetThumbnailPool())
+				.AllowCreate(true)
+				.OnShouldFilterAsset_Raw(this, &FCineAssemblyCustomization::ShouldFilterAssetBySchema, MetadataDesc.SchemaType)
+				.ObjectPath_Lambda([this, MetadataDesc]()
+					{
+						FString Value;
+						if (!CustomizedCineAssembly->GetMetadataAsString(MetadataDesc.Key, Value))
+						{
+							Value = MetadataDesc.DefaultValue.Get<FString>();
+						}
+						return Value;
+					})
+				.OnObjectChanged_Lambda([this, MetadataDesc](const FAssetData& InAssetData)
+					{
+						CustomizedCineAssembly->SetMetadataAsString(MetadataDesc.Key, InAssetData.GetObjectPathString());
+					});
+		default:
+			checkNoEntry();
+	}
+
+	return SNullWidget::NullWidget;
+}
+
 TSharedRef<SWidget> FCineAssemblyCustomization::MakeStringValueWidget(const FAssemblyMetadataDesc& MetadataDesc)
 {
 	TSharedPtr<SWidget> StringValueWidget;
diff --git a/Engine/Plugins/VirtualProduction/CinematicAssemblyTools/Source/CineAssemblyToolsEditor/Private/CineAssemblyCustomization.h b/Engine/Plugins/VirtualProduction/CinematicAssemblyTools/Source/CineAssemblyToolsEditor/Private/CineAssemblyCustomization.h
--- a/Engine/Plugins/VirtualProduction/CinematicAssemblyTools/Source/CineAssemblyToolsEditor/Private/CineAssemblyCustomization.h
+++ b/Engine/Plugins/VirtualProduction/CinematicAssemblyTools/Source/CineAssemblyToolsEditor/Private/CineAssemblyCustomization.h
@@ -32,7 +32,10 @@ private:
 	void CustomizeMetadataCategory(IDetailLayoutBuilder& DetailBuilder);
 	void CustomizeSubsequenceCategory(IDetailLayoutBuilder& DetailBuilder);
 
-	/** 
+	/** Creates the appropriate widget for a metadata field based on its type */
+	TSharedRef<SWidget> MakeMetadataWidget(IDetailLayoutBuilder& DetailBuilder, const FAssemblyMetadataDesc& MetadataDesc);
+
+	/**
 	* Given a FAssemblyMetadataDesc, makes a STemplateStringEditableTextBox for a string value that has bEvaluateTokens set to true,
 	* and a SMultiLineEditableTextBox for a string value that has bEvaluateTokens set to false.
 	*/
@@ -87,4 +90,7 @@ private:
 
 	/** Filter allowing cine assembly tokens to evaluate without the namespace prefix */
 	FNamingTokenFilterArgs FilterArgs;
+
+	/** True if the metadata properties are read-only, false if they are editable */
+	bool bIsMetadataReadOnly = true;
 };
