ShotManagement: Fix Assembly Duplication tool to properly duplicate all SubAssemblies recursively.
A recent change to better support SubAssembly creation moved the responsibility of creating SubAssembly objects to the parent Assembly, and responsibility for creating assets for those SubAssembly objects to the factory.
This change updates the logical flow of the duplication process to respect those changes to SubAssembly creation.
This change also accounts for the fact that SubAssemblies might now have their own SubAssemblies (previoulsy SubAssemblies could only be basic Assemblies with no schema), and so tokenized naming conventions and folder structures need to be accounted for recursively.

#rnx

#Synced-CL 51715153

diff --git a/Engine/Plugins/VirtualProduction/CinematicAssemblyTools/Source/CineAssemblyToolsEditor/Private/Sequencer/CineAssemblySchemaSequencerUtilities.cpp b/Engine/Plugins/VirtualProduction/CinematicAssemblyTools/Source/CineAssemblyToolsEditor/Private/Sequencer/CineAssemblySchemaSequencerUtilities.cpp
--- a/Engine/Plugins/VirtualProduction/CinematicAssemblyTools/Source/CineAssemblyToolsEditor/Private/Sequencer/CineAssemblySchemaSequencerUtilities.cpp
+++ b/Engine/Plugins/VirtualProduction/CinematicAssemblyTools/Source/CineAssemblyToolsEditor/Private/Sequencer/CineAssemblySchemaSequencerUtilities.cpp
@@ -4,6 +4,7 @@
 
 #include "CineCameraActor.h"
 #include "ISequencer.h"
+#include "Sections/MovieSceneSubSection.h"
 #include "SequencerUtilities.h"
 
 void FCineAssemblySchemaSequencerUtilities::CreateCamera(TSharedRef<ISequencer> Sequencer)
@@ -20,3 +21,14 @@ void FCineAssemblySchemaSequencerUtilities::CreateCamera(TSharedRef<ISequencer>
 
 	FSequencerUtilities::CreateBinding(Sequencer, *NewCamera, Params);
 }
+
+void FCineAssemblySchemaSequencerUtilities::ReplaceSubSequence(UMovieSceneSubSection* InSubSection, UMovieSceneSequence* InSequence)
+{
+	if (InSubSection)
+	{
+		const bool bIsLocked = InSubSection->IsLocked();
+		InSubSection->SetIsLocked(false);
+		InSubSection->SetSequence(InSequence);
+		InSubSection->SetIsLocked(bIsLocked);
+	}
+}
diff --git a/Engine/Plugins/VirtualProduction/CinematicAssemblyTools/Source/CineAssemblyToolsEditor/Private/Sequencer/CineAssemblySchemaSequencerUtilities.h b/Engine/Plugins/VirtualProduction/CinematicAssemblyTools/Source/CineAssemblyToolsEditor/Private/Sequencer/CineAssemblySchemaSequencerUtilities.h
--- a/Engine/Plugins/VirtualProduction/CinematicAssemblyTools/Source/CineAssemblyToolsEditor/Private/Sequencer/CineAssemblySchemaSequencerUtilities.h
+++ b/Engine/Plugins/VirtualProduction/CinematicAssemblyTools/Source/CineAssemblyToolsEditor/Private/Sequencer/CineAssemblySchemaSequencerUtilities.h
@@ -5,10 +5,19 @@
 #include "CoreMinimal.h"
 
 class ISequencer;
+class UMovieSceneSequence;
+class UMovieSceneSubSection;
 
 /** Utility functions for the custom Sequencer used by the Cine Assembly Schema */
 struct FCineAssemblySchemaSequencerUtilities
 {
 public:
+	/** Create a new spawnable camera binding without actually spawning a camera */
 	static void CreateCamera(TSharedRef<ISequencer> Sequencer);
+
+	/** 
+	 * Replace the Sequence in the input SubSection with the input Sequence
+	 * If needed, the SubSection will be first unlocked to ensure the assignment is succesful, and then its original lock status is restored.
+	 */
+	static void ReplaceSubSequence(UMovieSceneSubSection* InSubSection, UMovieSceneSequence* InSequence);
 };
diff --git a/Engine/Plugins/VirtualProduction/CinematicAssemblyTools/Source/CineAssemblyToolsEditor/Private/UI/CineAssembly/SDuplicateAssemblyWindow.cpp b/Engine/Plugins/VirtualProduction/CinematicAssemblyTools/Source/CineAssemblyToolsEditor/Private/UI/CineAssembly/SDuplicateAssemblyWindow.cpp
--- a/Engine/Plugins/VirtualProduction/CinematicAssemblyTools/Source/CineAssemblyToolsEditor/Private/UI/CineAssembly/SDuplicateAssemblyWindow.cpp
+++ b/Engine/Plugins/VirtualProduction/CinematicAssemblyTools/Source/CineAssemblyToolsEditor/Private/UI/CineAssembly/SDuplicateAssemblyWindow.cpp
@@ -14,6 +14,7 @@
 #include "ObjectTools.h"
 #include "SCineAssemblyConfigPanel.h"
 #include "Sections/MovieSceneSubSection.h"
+#include "Sequencer/CineAssemblySchemaSequencerUtilities.h"
 #include "SNamingTokensEditableTextBox.h"
 #include "Tracks/MovieSceneSubTrack.h"
 #include "Widgets/Input/SEditableTextBox.h"
@@ -29,23 +30,28 @@ const FName SDuplicateAssemblyWindow::DuplicateNameColumn = "DuplicateName";
 
 const FString SDuplicateAssemblyWindow::DuplicationPreferenceSection = TEXT("CinematicAssemblyDuplication_");
 
+UCineAssembly* FSequenceTreeItem::GetDuplicateSubAssembly() const
+{
+	if (bIsManagedSubAssembly)
+	{
+		return Cast<UCineAssembly>(SubSection->GetSequence());
+	}
+	return nullptr;
+}
+
 void SDuplicateAssemblyWindow::Construct(const FArguments& InArgs, UCineAssembly* InAssembly)
 {
 	OriginalAssembly = InAssembly;
-	Schema = OriginalAssembly->GetSchema();
 
 	// Duplicate the original assembly into a new transient assembly that can be configured in the UI
 	// Note: all of the subsequence tracks/sections will be duplicated, but will reference the sequences in the original assembly
 	// So if the user chooses MaintainReference as the duplication mode for any subsequence, no further work is needed for that subsequence
 	DuplicateAssembly.Reset(Cast<UCineAssembly>(StaticDuplicateObject(OriginalAssembly, GetTransientPackage(), NAME_None, RF_Transient)));
 
-	// Get the path and root folder of the original assembly) and initialize the duplication path to be the same root folder as the original
+	// Get the path and root folder of the original assembly and initialize the duplication path to be the same root folder as the original
 	OriginalAssembly->GetAssetPathAndRootFolder(OriginalAssemblyPath, OriginalAssemblyRootFolder);
 	DuplicationPath = OriginalAssemblyRootFolder;
 
-	// This path will be used to find the relative paths for some subsequences, but the function that finds the relative path only works if there is a trailing '/'
-	OriginalAssemblyPath = OriginalAssemblyPath.AppendChar('/');
-
 	const FVector2D DefaultWindowSize = FVector2D(1200, 750);
 
 	SWindow::Construct(SWindow::FArguments()
@@ -99,8 +105,18 @@ void SDuplicateAssemblyWindow::Construct(const FArguments& InArgs, UCineAssembly
 
 TSharedRef<SWidget> SDuplicateAssemblyWindow::MakeSubsequencesPanel()
 {
-	// Walk the duplicate assembly to recursively find every subsequence and build up the subsequence tree
-	PopulateSubsequenceTreeRecurisve(DuplicateAssembly.Get(), nullptr);
+	UMovieScene* DuplicateAssemblyMovieScene = DuplicateAssembly->GetMovieScene();
+	if (DuplicateAssemblyMovieScene)
+	{
+		// In order to duplicate subsequences, the MovieScene needs to be mutable. The original read-only state will be restored when we are done.
+		const bool bRootIsReadOnly = DuplicateAssemblyMovieScene->IsReadOnly();
+		DuplicateAssemblyMovieScene->SetReadOnly(false);
+
+		// Walk the duplicate assembly to recursively find every subsequence and build up the subsequence tree
+		PopulateSubsequenceTreeRecursive(DuplicateAssembly.Get(), nullptr, OriginalAssemblyPath);
+
+		DuplicateAssemblyMovieScene->SetReadOnly(bRootIsReadOnly);
+	}
 
 	TreeView = SNew(STreeView<TSharedPtr<FSequenceTreeItem>>)
 		.TreeItemsSource(&TreeItems)
@@ -216,7 +232,7 @@ TSharedRef<SWidget> SDuplicateAssemblyWindow::MakeSubsequencesPanel()
 								]
 								.ButtonContent()
 								[
-	 								SNew(SImage).Image(FCineAssemblyToolsStyle::Get().GetBrush("Icons.Folder"))
+									SNew(SImage).Image(FCineAssemblyToolsStyle::Get().GetBrush("Icons.Folder"))
 								]
 						]
 				]
@@ -260,7 +276,13 @@ TSharedRef<SWidget> SDuplicateAssemblyWindow::MakeButtonsPanel()
 
 FReply SDuplicateAssemblyWindow::OnDuplicateClicked()
 {
-	UCineAssemblyFactory::CreateConfiguredAssembly(DuplicateAssembly.Get(), DuplicationPath);
+	SaveDuplicationPreferences();
+
+	// Managed SubAssemblies were pre-duplicated during tree construction. Restore the original sequence in any SubSections that are not set to DuplicateOriginal
+	for (TSharedPtr<FSequenceTreeItem>& Item : TreeItems)
+	{
+		RestoreOriginalSubAssembliesRecursive(Item);
+	}
 
 	// Flatten the duplication data from the subsequence tree into a map
 	for (TSharedPtr<FSequenceTreeItem>& Item : TreeItems)
@@ -268,17 +290,15 @@ FReply SDuplicateAssemblyWindow::OnDuplicateClicked()
 		BuildDuplicationMapRecursive(Item);
 	}
 
-	DuplicateAssembly->GetAssetPathAndRootFolder(DuplicateAssemblyPath, DuplicateAssemblyRootFolder);
+	// Create valid packages for the Duplicate Assembly and all of its duplicated SubAssemblies
+	UCineAssemblyFactory::CreateConfiguredAssembly(DuplicateAssembly.Get(), DuplicationPath);
 
-	// Duplicate any subsequences that are set to DuplicateOriginal
+	// Duplicate the remaining subsequences (those that are not Assemblies) that are set to DuplicateOriginal
 	DuplicateSubsequencesRecursive(DuplicateAssembly.Get());
 
-	// Remove any subsequencers that are set to Remove
+	// Remove any subsequences that are set to Remove
 	RemoveSubsequencesRecursive(DuplicateAssembly.Get());
 
-	// Write out the duplication mode values for each subassembly to a config file so those values can be used as the default the next time an assembly with this schema type is duplicated.
-	SaveDuplicationPreferences();
-
 	RequestDestroyWindow();
 	return FReply::Handled();
 }
@@ -299,7 +319,7 @@ void SDuplicateAssemblyWindow::OnGetChildren(TSharedPtr<FSequenceTreeItem> TreeI
 	OutNodes.Append(TreeItem->Children);
 }
 
-void SDuplicateAssemblyWindow::PopulateSubsequenceTreeRecurisve(UMovieSceneSequence* InSequence, TSharedPtr<FSequenceTreeItem> InItem)
+void SDuplicateAssemblyWindow::PopulateSubsequenceTreeRecursive(UMovieSceneSequence* InSequence, TSharedPtr<FSequenceTreeItem> InItem, const FString& SourceBasePath)
 {
 	UMovieScene* MovieScene = InSequence->GetMovieScene();
 	if (!MovieScene)
@@ -307,64 +327,70 @@ void SDuplicateAssemblyWindow::PopulateSubsequenceTreeRecurisve(UMovieSceneSeque
 		return;
 	}
 
-	// Loop through all of the tracks and all of the sections to find every subsequence in the input sequence
-	for (UMovieSceneTrack* Track : MovieScene->GetTracks())
+	// Loop through all of the sections to find every subsequence in the input sequence
+	for (UMovieSceneSection* Section : MovieScene->GetAllSections())
 	{
-		if (UMovieSceneSubTrack* SubTrack = Cast<UMovieSceneSubTrack>(Track))
+		if (UMovieSceneSubSection* SubSection = Cast<UMovieSceneSubSection>(Section))
 		{
-			for (UMovieSceneSection* Section : SubTrack->GetAllSections())
+			if (UMovieSceneSequence* Subsequence = SubSection->GetSequence())
 			{
-				if (UMovieSceneSubSection* SubSection = Cast<UMovieSceneSubSection>(Section))
+				TSharedPtr<FSequenceTreeItem> NewItem = MakeShared<FSequenceTreeItem>();
+				NewItem->SubSection = SubSection;
+				NewItem->OriginalSequence = Subsequence;
+				NewItem->Parent = InItem;
+
+				NewItem->DuplicationData.SubsequenceName = Subsequence->IsA<UCineAssembly>() ? Cast<UCineAssembly>(Subsequence)->AssemblyName.Template : Subsequence->GetDisplayName().ToString();
+
+				InItem ? InItem->Children.Add(NewItem) : TreeItems.Add(NewItem);
+
+				// We handle managed SubAssemblies differently than other Subsequences found in the input Sequence's MovieScene 
+				UCineAssembly* OwningAssembly = Cast<UCineAssembly>(InSequence);
+				if (OwningAssembly && OwningAssembly->SubAssemblies.Contains(SubSection))
 				{
-					if (UMovieSceneSequence* Subsequence = SubSection->GetSequence())
-					{
-						TSharedPtr<FSequenceTreeItem> NewItem = MakeShared<FSequenceTreeItem>();
-						NewItem->Subsequence = Subsequence;
+					// Pre-duplicate this SubAssembly such that it can be configured by the user before assets are created
+					// If the user chooses to duplicate this, a valid asset will be created by the asset Factory
+					// Otherwise, the original SubSequence will be restored and this duplicate will be thrown out
+					UCineAssembly* DuplicateSubAssembly = Cast<UCineAssembly>(StaticDuplicateObject(Subsequence, GetTransientPackage(), NAME_None, RF_Transient));
 
-						// Check if this subsequence is one of the subassemblies created from the schema
-						const int32 SubAssemblyIndex = DuplicateAssembly->SubAssemblies.IndexOfByKey(SubSection);
+					DuplicateSubAssembly->ParentAssembly = OwningAssembly;
 
-						if (SubAssemblyIndex != INDEX_NONE)
-						{
-							// If this is a subassembly, use the template name (originally defined by the schema)
-							NewItem->DuplicationData.bIsSubAssembly = true;
-							NewItem->DuplicationData.SubsequenceName.Template = DuplicateAssembly->SubAssemblyNames[SubAssemblyIndex].Template;
-							NewItem->DuplicationData.DuplicationMode = GetDuplicationPreference(NewItem->DuplicationData.SubsequenceName.Template);
-						}
-						else
-						{
-							// If this is not a subassembly, find the relative path to the original assembly so that a duplicate susbequence can be created with the same relative folder structure
-							NewItem->DuplicationData.bIsSubAssembly = false;
+					FCineAssemblySchemaSequencerUtilities::ReplaceSubSequence(SubSection, DuplicateSubAssembly);
 
-							FString SubsequencePath = FPaths::GetPath(Subsequence->GetPathName());
-							FPaths::MakePathRelativeTo(SubsequencePath, *OriginalAssemblyPath);
+					NewItem->bIsManagedSubAssembly = true;
+					NewItem->DuplicationData.DuplicationMode = GetDuplicationPreference(DuplicateSubAssembly->AssemblyName.Template);
+				}
+				else
+				{
+					// Compute the path relative to the parent Assembly so that a duplicate SubSequence can be created in the correct relative folder structure
+					FString SubsequencePath = FPaths::GetPath(Subsequence->GetPathName());
+					const FString RelativeParentPath = SourceBasePath + TEXT('/');
+					FPaths::MakePathRelativeTo(SubsequencePath, *RelativeParentPath);
+					NewItem->DuplicationData.RelativePath = SubsequencePath;
+				}
 
-							NewItem->DuplicationData.SubsequenceName.Template = SubsequencePath / Subsequence->GetDisplayName().ToString();
-							NewItem->DuplicationData.DuplicationMode = ECineAssemblyDuplicationMode::DuplicateOriginal;
-						}
+				// Query for the current SubSequence again (which could be the new duplicate SubAssembly) to ensure we recurse into the correct sequence
+				UMovieSceneSequence* SequenceToRecurse = SubSection->GetSequence();
+				if (UMovieScene* SubSequenceMovieScene = SequenceToRecurse->GetMovieScene())
+				{
+					// In order to duplicate subsequences, the MovieScene needs to be mutable. The original read-only state will be restored when we are done.
+					const bool bMovieSceneIsReadOnly = SubSequenceMovieScene->IsReadOnly();
+					SubSequenceMovieScene->SetReadOnly(false);
 
-						if (InItem)
-						{
-							InItem->Children.Add(NewItem);
-							NewItem->Parent = InItem;
-						}
-						else
-						{
-							TreeItems.Add(NewItem);
-							NewItem->Parent = nullptr;
-						}
+					// Use the path of the original Subsequence (even if it was duplicated already) because the duplicate does not yet have a valid path
+					const FString ChildSourceBasePath = FPaths::GetPath(Subsequence->GetPathName());
 
-						PopulateSubsequenceTreeRecurisve(Subsequence, NewItem);
-					}
+					PopulateSubsequenceTreeRecursive(SequenceToRecurse, NewItem, ChildSourceBasePath);
+
+					SubSequenceMovieScene->SetReadOnly(bMovieSceneIsReadOnly);
 				}
 			}
 		}
 	}
 }
 
-void SDuplicateAssemblyWindow::DuplicateSubsequencesRecursive(UMovieSceneSequence* Sequence)
+void SDuplicateAssemblyWindow::DuplicateSubsequencesRecursive(UMovieSceneSequence* InSequence)
 {
-	UMovieScene* MovieScene = Sequence->GetMovieScene();
+	UMovieScene* MovieScene = InSequence->GetMovieScene();
 	if (!MovieScene)
 	{
 		return;
@@ -377,68 +403,44 @@ void SDuplicateAssemblyWindow::DuplicateSubsequencesRecursive(UMovieSceneSequenc
 
 	IAssetTools& AssetTools = FModuleManager::GetModuleChecked<FAssetToolsModule>("AssetTools").Get();
 
-	// Loop through all of the tracks and all of the sections to find every subsequence in the input sequence
-	for (UMovieSceneTrack* Track : MovieScene->GetTracks())
+	// Loop through all of the sections to find every subsequence in the input sequence
+	for (UMovieSceneSection* Section : MovieScene->GetAllSections())
 	{
-		if (UMovieSceneSubTrack* SubTrack = Cast<UMovieSceneSubTrack>(Track))
+		if (UMovieSceneSubSection* SubSection = Cast<UMovieSceneSubSection>(Section))
 		{
-			for (UMovieSceneSection* Section : SubTrack->GetAllSections())
+			if (UMovieSceneSequence* Subsequence = SubSection->GetSequence())
 			{
-				if (UMovieSceneSubSection* SubSection = Cast<UMovieSceneSubSection>(Section))
+				// Only duplicate subsequences that are in the map and set to DuplicateOriginal.
+				// Note: Managed SubAssemblies that were already duplicated were not added to the map to avoid double duplication
+				if (const FSubsequenceDuplicationData* DuplicationData = SubsequenceDuplicationData.Find(Subsequence))
 				{
-					if (UMovieSceneSequence* Subsequence = SubSection->GetSequence())
+					if (DuplicationData->DuplicationMode == ECineAssemblyDuplicationMode::DuplicateOriginal)
 					{
-						if (!SubsequenceDuplicationData.Contains(Subsequence))
-						{
-							continue;
-						}
+						const FString ResolvedName = UCineAssemblyNamingTokens::GetResolvedText(DuplicationData->SubsequenceName, DuplicateAssembly.Get()).ToString();
+
+						// Use the current parent sequence's path as the base, so that subsequences inside a renamed SubAssembly are placed relative to the SubAssembly's new folder
+						const FString ParentAssetPath = FPaths::GetPath(InSequence->GetPathName());
+						FString DesiredAssetName = ParentAssetPath / DuplicationData->RelativePath / ResolvedName;
+						FPaths::CollapseRelativeDirectories(DesiredAssetName);
+
+						FString UniquePackageName;
+						FString UniqueAssetName;
+						AssetTools.CreateUniqueAssetName(DesiredAssetName, TEXT(""), UniquePackageName, UniqueAssetName);
 
-						const FSubsequenceDuplicationData& DuplicationData = SubsequenceDuplicationData[Subsequence];
+						const FString DestinationFolder = FPaths::GetPath(UniquePackageName);
+						UMovieSceneSequence* DuplicateSubsequence = Cast<UMovieSceneSequence>(AssetTools.DuplicateAsset(UniqueAssetName, DestinationFolder, Subsequence));
 
-						if (DuplicationData.DuplicationMode == ECineAssemblyDuplicationMode::DuplicateOriginal)
+						if (DuplicateSubsequence)
 						{
-							const FString SubsequenceNameAndPath = UCineAssemblyNamingTokens::GetResolvedText(DuplicationData.SubsequenceName.Template, DuplicateAssembly.Get()).ToString();
-
-							FString DesiredAssetName;
-							if (DuplicationData.bIsSubAssembly)
-							{
-								// Update the display name of the subsequence track to match the resolved subassembly name
-								const FString SubAssemblyFilename = FPaths::GetBaseFilename(SubsequenceNameAndPath);
-								SubTrack->SetDisplayName(FText::FromString(SubAssemblyFilename));
-
-								DesiredAssetName = DuplicateAssemblyRootFolder / SubsequenceNameAndPath;
-							}
-							else
-							{
-								DesiredAssetName = DuplicateAssemblyPath / SubsequenceNameAndPath;
-								FPaths::CollapseRelativeDirectories(DesiredAssetName);
-							}
-
-							FString UniquePackageName;
-							FString UniqueAssetName;
-							AssetTools.CreateUniqueAssetName(DesiredAssetName, TEXT(""), UniquePackageName, UniqueAssetName);
-
-							const FString SubAssemblyPath = FPaths::GetPath(UniquePackageName);
-							UMovieSceneSequence* DuplicateSubsequence = Cast<UMovieSceneSequence>(AssetTools.DuplicateAsset(UniqueAssetName, SubAssemblyPath, Subsequence));
-
-							if (DuplicateSubsequence)
-							{
-								// Update the sequence of the duplicate subsection to be the duplicate subsequence
-								// If the section is locked, unlock it, then update the subsequence, then lock it again
-								const bool bIsLocked = SubSection->IsLocked();
-								SubSection->SetIsLocked(false);
-								SubSection->SetSequence(DuplicateSubsequence);
-								SubSection->SetIsLocked(bIsLocked);
-
-								// Update the data in the duplication data map to correspond to the new duplicate subsequence
-								SubsequenceDuplicationData.Add(DuplicateSubsequence, DuplicationData);
-								SubsequenceDuplicationData.Remove(Subsequence);
-							}
-						}
+							FCineAssemblySchemaSequencerUtilities::ReplaceSubSequence(SubSection, DuplicateSubsequence);
 
-						DuplicateSubsequencesRecursive(SubSection->GetSequence());
+							SubsequenceDuplicationData.Add(DuplicateSubsequence, *DuplicationData);
+							SubsequenceDuplicationData.Remove(Subsequence);
+						}
 					}
 				}
+
+				DuplicateSubsequencesRecursive(SubSection->GetSequence());
 			}
 		}
 	}
@@ -455,8 +457,6 @@ void SDuplicateAssemblyWindow::RemoveSubsequencesRecursive(UMovieSceneSequence*
 		return;
 	}
 
-	IAssetTools& AssetTools = FModuleManager::GetModuleChecked<FAssetToolsModule>("AssetTools").Get();
-
 	TArray<UMovieSceneTrack*> TracksToRemove;
 
 	// Loop through all of the tracks and all of the sections to find every subsequence in the input sequence
@@ -472,13 +472,12 @@ void SDuplicateAssemblyWindow::RemoveSubsequencesRecursive(UMovieSceneSequence*
 				{
 					if (UMovieSceneSequence* Subsequence = SubSection->GetSequence())
 					{
-						if (!SubsequenceDuplicationData.Contains(Subsequence))
+						const FSubsequenceDuplicationData* DuplicationData = SubsequenceDuplicationData.Find(Subsequence);
+						if (!DuplicationData)
 						{
 							continue;
 						}
-
-						const FSubsequenceDuplicationData& DuplicationData = SubsequenceDuplicationData[Subsequence];
-						if (DuplicationData.DuplicationMode == ECineAssemblyDuplicationMode::Remove)
+						if (DuplicationData->DuplicationMode == ECineAssemblyDuplicationMode::Remove)
 						{
 							SectionsToRemove.Add(SubSection);
 						}
@@ -528,9 +527,27 @@ void SDuplicateAssemblyWindow::RemoveTracksFromFoldersRecursive(TArrayView<UMovi
 	}
 }
 
+void SDuplicateAssemblyWindow::RestoreOriginalSubAssembliesRecursive(const TSharedPtr<FSequenceTreeItem>& InItem)
+{
+	if (InItem->bIsManagedSubAssembly && InItem->DuplicationData.DuplicationMode != ECineAssemblyDuplicationMode::DuplicateOriginal)
+	{
+		FCineAssemblySchemaSequencerUtilities::ReplaceSubSequence(InItem->SubSection, InItem->OriginalSequence);
+	}
+
+	for (const TSharedPtr<FSequenceTreeItem>& Child : InItem->Children)
+	{
+		RestoreOriginalSubAssembliesRecursive(Child);
+	}
+}
+
 void SDuplicateAssemblyWindow::BuildDuplicationMapRecursive(const TSharedPtr<FSequenceTreeItem> InItem)
 {
-	SubsequenceDuplicationData.Add(InItem->Subsequence, InItem->DuplicationData);
+	// SubAssemblies have already been duplicated so they should not be added to the map to avoid being duplicated again
+	const bool bIsAlreadyDuplicated = InItem->bIsManagedSubAssembly && (InItem->DuplicationData.DuplicationMode == ECineAssemblyDuplicationMode::DuplicateOriginal);
+	if (!bIsAlreadyDuplicated)
+	{
+		SubsequenceDuplicationData.Add(InItem->OriginalSequence, InItem->DuplicationData);
+	}
 
 	for (const TSharedPtr<FSequenceTreeItem>& Child : InItem->Children)
 	{
@@ -550,28 +567,40 @@ void SDuplicateAssemblyWindow::ExpandTreeItem(TSharedPtr<FSequenceTreeItem> InIt
 
 void SDuplicateAssemblyWindow::SaveDuplicationPreferences() const
 {
+	const UCineAssemblySchema* Schema = DuplicateAssembly->GetSchema();
 	if (Schema && GConfig)
 	{
 		const FString SectionName = DuplicationPreferenceSection + Schema->SchemaName;
 
-		for (const TPair<UMovieSceneSequence*, FSubsequenceDuplicationData>& DuplicationData : SubsequenceDuplicationData)
+		for (const TSharedPtr<FSequenceTreeItem>& Item : TreeItems)
 		{
-			if (DuplicationData.Value.bIsSubAssembly)
-			{
-				// The config strips off these token characters (for some reason), so we save and restore the values without them
-				FString TemplateName = DuplicationData.Value.SubsequenceName.Template;
-				TemplateName = TemplateName.Replace(TEXT("{"), TEXT(""));
-				TemplateName = TemplateName.Replace(TEXT("}"), TEXT(""));
-
-				const int32 Value = static_cast<int32>(DuplicationData.Value.DuplicationMode);
-				GConfig->SetInt(*SectionName, *TemplateName, Value, GEditorPerProjectIni);
-			}
+			SaveDuplicationPreferencesRecursive(Item, SectionName);
 		}
 	}
 }
 
+void SDuplicateAssemblyWindow::SaveDuplicationPreferencesRecursive(const TSharedPtr<FSequenceTreeItem>& InItem, const FString& SectionName) const
+{
+	if (UCineAssembly* DuplicateSubAssembly = InItem->GetDuplicateSubAssembly())
+	{
+		// The config strips off these token characters (for some reason), so we save and restore the values without them
+		FString TemplateName = DuplicateSubAssembly->AssemblyName.Template;
+		TemplateName = TemplateName.Replace(TEXT("{"), TEXT(""));
+		TemplateName = TemplateName.Replace(TEXT("}"), TEXT(""));
+
+		const int32 Value = static_cast<int32>(InItem->DuplicationData.DuplicationMode);
+		GConfig->SetInt(*SectionName, *TemplateName, Value, GEditorPerProjectIni);
+	}
+
+	for (const TSharedPtr<FSequenceTreeItem>& Child : InItem->Children)
+	{
+		SaveDuplicationPreferencesRecursive(Child, SectionName);
+	}
+}
+
 ECineAssemblyDuplicationMode SDuplicateAssemblyWindow::GetDuplicationPreference(const FString& SubAssemblyName) const
 {
+	const UCineAssemblySchema* Schema = DuplicateAssembly->GetSchema();
 	if (Schema && GConfig)
 	{
 		const FString SectionName = DuplicationPreferenceSection + Schema->SchemaName;
@@ -595,13 +624,21 @@ void SSubsequenceDuplicationRow::Construct(const FArguments& Args, const TShared
 	TreeItem = InTreeItem;
 
 	NamingTokenContext = TStrongObjectPtr<UCineAssemblyNamingTokensContext>(NewObject<UCineAssemblyNamingTokensContext>());
-	NamingTokenContext->Assembly = DuplicateAssembly;
+
+	if (UCineAssembly* DuplicateSubAssembly = TreeItem->GetDuplicateSubAssembly())
+	{
+		NamingTokenContext->Assembly = DuplicateSubAssembly;
+	}
+	else
+	{
+		NamingTokenContext->Assembly = DuplicateAssembly;
+	}
 
 	FilterArgs.AdditionalNamespacesToInclude.Add(UCineAssemblyNamingTokens::TokenNamespace);
 
 	FSuperRowType::FArguments StyleArguments = FSuperRowType::FArguments()
 		.Padding(FMargin(4.0f, 8.0f))
-		.Style(&FAppStyle().GetWidgetStyle<FTableRowStyle>("TableView.AlternatingRow"));
+		.Style(&FAppStyle::Get().GetWidgetStyle<FTableRowStyle>("TableView.AlternatingRow"));
 
 	SMultiColumnTableRow<TSharedPtr<FSequenceTreeItem>>::Construct(StyleArguments, OwnerTableView);
 }
@@ -647,7 +684,7 @@ TSharedRef<SWidget> SSubsequenceDuplicationRow::MakeOriginalNameWidget()
 		.VAlign(VAlign_Center)
 		.Padding(8.0f, 0.0f, 0.0f, 0.0f)
 		[
-			SNew(STextBlock).Text(TreeItem->Subsequence->GetDisplayName())
+			SNew(STextBlock).Text(TreeItem->OriginalSequence->GetDisplayName())
 		];
 }
 
@@ -690,9 +727,9 @@ TSharedRef<SWidget> SSubsequenceDuplicationRow::MakeDuplicationModeWidget()
 				.IsEnabled_Lambda([this]()
 					{
 						// Disable the duplication mode widget if the parent is MaintainReference or Remove (since the value should be locked to match the parent)
-						if (TreeItem->Parent && TreeItem->Parent->DuplicationData.DuplicationMode != ECineAssemblyDuplicationMode::DuplicateOriginal)
+						if (TSharedPtr<FSequenceTreeItem> ParentPinned = TreeItem->Parent.Pin())
 						{
-							return false;
+							return ParentPinned->DuplicationData.DuplicationMode == ECineAssemblyDuplicationMode::DuplicateOriginal;
 						}
 						return true;
 					})
@@ -723,10 +760,9 @@ TSharedRef<SWidget> SSubsequenceDuplicationRow::MakeDuplicateNameWidget()
 					{
 						if (TreeItem->DuplicationData.DuplicationMode == ECineAssemblyDuplicationMode::DuplicateOriginal)
 						{
-							const FString SubsequenceName = FPaths::GetBaseFilename(TreeItem->DuplicationData.SubsequenceName.Template);
-							return FText::FromString(SubsequenceName);
+							return FText::FromString(TreeItem->DuplicationData.SubsequenceName);
 						}
-						return TreeItem->Subsequence->GetDisplayName();
+						return TreeItem->OriginalSequence->GetDisplayName();
 					})
 				.OnTextCommitted(this, &SSubsequenceDuplicationRow::OnSubsequenceNameCommitted)
 				.OnTokenizedTextEvaluated(this, &SSubsequenceDuplicationRow::OnSubsequenceTokenTextEvaluated)
@@ -765,21 +801,19 @@ void SSubsequenceDuplicationRow::SetDuplicationModeRecursive(TSharedPtr<FSequenc
 
 void SSubsequenceDuplicationRow::OnSubsequenceNameCommitted(const FText& InText, ETextCommit::Type InCommitType)
 {
-	if (TreeItem->DuplicationData.DuplicationMode == ECineAssemblyDuplicationMode::DuplicateOriginal)
+	TreeItem->DuplicationData.SubsequenceName = InText.ToString();
+
+	if (UCineAssembly* DuplicateSubAssembly = TreeItem->GetDuplicateSubAssembly())
 	{
-		const FString OldPath = TreeItem->DuplicationData.SubsequenceName.Template;
-		const FString NewPath = FPaths::GetPath(OldPath) / InText.ToString();
-		TreeItem->DuplicationData.SubsequenceName.Template = NewPath;
+		DuplicateSubAssembly->AssemblyName.Template = InText.ToString();
 	}
 }
 
 void SSubsequenceDuplicationRow::OnSubsequenceTokenTextEvaluated(const FText& InText)
 {
-	if (TreeItem->DuplicationData.DuplicationMode == ECineAssemblyDuplicationMode::DuplicateOriginal)
+	if (UCineAssembly* DuplicateSubAssembly = TreeItem->GetDuplicateSubAssembly())
 	{
-		const FString OldPath = TreeItem->DuplicationData.SubsequenceName.Template;
-		const FString NewPath = FPaths::GetPath(OldPath) / InText.ToString();
-		TreeItem->DuplicationData.SubsequenceName.Resolved = FText::FromString(NewPath);
+		DuplicateSubAssembly->AssemblyName.Resolved = InText;
 	}
 }
 
diff --git a/Engine/Plugins/VirtualProduction/CinematicAssemblyTools/Source/CineAssemblyToolsEditor/Private/UI/CineAssembly/SDuplicateAssemblyWindow.h b/Engine/Plugins/VirtualProduction/CinematicAssemblyTools/Source/CineAssemblyToolsEditor/Private/UI/CineAssembly/SDuplicateAssemblyWindow.h
--- a/Engine/Plugins/VirtualProduction/CinematicAssemblyTools/Source/CineAssemblyToolsEditor/Private/UI/CineAssembly/SDuplicateAssemblyWindow.h
+++ b/Engine/Plugins/VirtualProduction/CinematicAssemblyTools/Source/CineAssemblyToolsEditor/Private/UI/CineAssembly/SDuplicateAssemblyWindow.h
@@ -11,6 +11,7 @@
 #include "Widgets/Views/STreeView.h"
 
 class UMovieSceneFolder;
+class UMovieSceneSubSection;
 
 /** List of available duplication modes for a subsequence */
 enum class ECineAssemblyDuplicationMode : uint8
@@ -26,27 +27,36 @@ struct FSubsequenceDuplicationData
 	/** The duplication mode to use */
 	ECineAssemblyDuplicationMode DuplicationMode = ECineAssemblyDuplicationMode::DuplicateOriginal;
 
-	/** The name to use for the duplicated subsequence (if the mode if DuplicateOriginal) */
-	FTemplateString SubsequenceName;
+	/** The folder path relative to the parent sequence */
+	FString RelativePath;
 
-	/** Whether or not the subsequence was defined in the schema of the original subsequence */
-	bool bIsSubAssembly = false;
+	/** The name to use for the duplicated subsequence (if the mode is DuplicateOriginal) */
+	FString SubsequenceName;
 };
 
 /** An item in the subsequence tree view, associating a subsequence with its duplication data and child subsequences */
 struct FSequenceTreeItem
 {
-	/** The subsequence represented by this item in the tree */
-	UMovieSceneSequence* Subsequence;
+	/** The SubSection in the MovieScene of the Duplicate Assembly (or one of its SubAssemblies) */
+	UMovieSceneSubSection* SubSection = nullptr;
 
-	/** The duplication associated with this subsequence */
+	/** The original sequence referenced by this SubSection before any duplication occured */
+	UMovieSceneSequence* OriginalSequence = nullptr;
+
+	/** The duplication data associated with this subsequence */
 	FSubsequenceDuplicationData DuplicationData;
 
 	/** Child subsequence tree items */
 	TArray<TSharedPtr<FSequenceTreeItem>> Children;
 
 	/** Parent sequence tree item */
-	TSharedPtr<FSequenceTreeItem> Parent;
+	TWeakPtr<FSequenceTreeItem> Parent;
+
+	/** Indicates whether this tree item represents a SubAssembly in the managed SubAssemblies list of an Assembly being duplicated */
+	bool bIsManagedSubAssembly = false;
+
+	/** If this tree item is a managed SubAssembly, returns the duplicate SubAssembly in the SubSection. Otherwise, returns nullptr */
+	UCineAssembly* GetDuplicateSubAssembly() const;
 };
 
 /** Widget row for an item in the subsequence tree view */
@@ -129,12 +139,18 @@ private:
 	/** Adds the children of the input tree item to the OutNodes array */
 	void OnGetChildren(TSharedPtr<FSequenceTreeItem> InTreeItem, TArray<TSharedPtr<FSequenceTreeItem>>& OutNodes);
 
-	/** Builds the list of tree items in the subsequence tree by recursively walking the input sequence and finding all subsequences that may need to be duplicated */
-	void PopulateSubsequenceTreeRecurisve(UMovieSceneSequence* InSequence, TSharedPtr<FSequenceTreeItem> InTreeItem);
+	/** 
+	 * Builds the list of tree items in the subsequence tree by recursively walking the input sequence and finding all subsequences that may need to be duplicated.
+	 * SourceBasePath is the asset folder used to compute non-SubAssembly relative paths. 
+	 */
+	void PopulateSubsequenceTreeRecursive(UMovieSceneSequence* InSequence, TSharedPtr<FSequenceTreeItem> InTreeItem, const FString& SourceBasePath);
 
 	/** Recursively duplicates any of the subsequences of the input sequence that are set to the DuplicateOriginal duplication mode */
 	void DuplicateSubsequencesRecursive(UMovieSceneSequence* InSequence);
 
+	/** Recursively restores the original subsequences of the input item that are not set to DuplicateOriginal duplication mode */
+	void RestoreOriginalSubAssembliesRecursive(const TSharedPtr<FSequenceTreeItem>& InItem);
+
 	/** Recursively removes any of the subsequences of the input sequence that are set to the Remove duplication mode */
 	void RemoveSubsequencesRecursive(UMovieSceneSequence* InSequence);
 
@@ -150,6 +166,9 @@ private:
 	/** Save the duplication mode settings for each subassembly to a config file to re-use the next time an assembly with the same schema is duplicated */
 	void SaveDuplicationPreferences() const;
 
+	/** Recursively saves the duplication mode preference for the input tree item and its children */
+	void SaveDuplicationPreferencesRecursive(const TSharedPtr<FSequenceTreeItem>& InItem, const FString& SectionName) const;
+
 	/** Reads the saved config and returns the most recent duplication mode for the input subassembly */
 	ECineAssemblyDuplicationMode GetDuplicationPreference(const FString& SubAssemblyName) const;
 
@@ -157,9 +176,6 @@ private:
 	/** The original assembly being duplicated */
 	UCineAssembly* OriginalAssembly;
 
-	/** The schema of the assembly being duplicated */
-	const UCineAssemblySchema* Schema;
-
 	/** The duplicated assembly that is being configured */
 	TStrongObjectPtr<UCineAssembly> DuplicateAssembly;
 
@@ -169,7 +185,7 @@ private:
 	/** Tree view of subsequences, allowing the user to configure the duplication data */
 	TSharedPtr<STreeView<TSharedPtr<FSequenceTreeItem>>> TreeView;
 
-	/** Map of subsequences to the associated duplication data, which mirrors the tree view */
+	/** Subset of the duplication data derived from the tree view for non-managed subsequences */
 	TMap<UMovieSceneSequence*, FSubsequenceDuplicationData> SubsequenceDuplicationData;
 
 	/** The path where the duplicate assembly asset should be created */
@@ -181,12 +197,6 @@ private:
 	/** The root folder of the original assembly (i.e. the asset path with the default assembly path from the schema removed) */
 	FString OriginalAssemblyRootFolder;
 
-	/** The path of the duplicate assembly asset (after it has been renamed by the factory) */
-	FString DuplicateAssemblyPath;
-
-	/** The root folder of the duplicate assembly (i.e. the asset path with the default assembly path from the schema removed) */
-	FString DuplicateAssemblyRootFolder;
-
 	/** Config section name for duplication preferences */
 	static const FString DuplicationPreferenceSection;
 };
