ShotManagement: Add DuplicateAssembly function for editor scripting.
Move AddMetadataToAssembly from the Production Function Library to the Cine Assembly Function Library, and also expose it to editor scripting.
#rnx

#Synced-CL 51715153

diff --git a/Engine/Plugins/VirtualProduction/CinematicAssemblyTools/Source/CineAssemblyToolsEditor/Private/CineAssemblyEditorFunctionLibrary.cpp b/Engine/Plugins/VirtualProduction/CinematicAssemblyTools/Source/CineAssemblyToolsEditor/Private/CineAssemblyEditorFunctionLibrary.cpp
--- a/Engine/Plugins/VirtualProduction/CinematicAssemblyTools/Source/CineAssemblyToolsEditor/Private/CineAssemblyEditorFunctionLibrary.cpp
+++ b/Engine/Plugins/VirtualProduction/CinematicAssemblyTools/Source/CineAssemblyToolsEditor/Private/CineAssemblyEditorFunctionLibrary.cpp
@@ -2,10 +2,160 @@
 
 #include "CineAssemblyEditorFunctionLibrary.h"
 
+#include "AssetToolsModule.h"
 #include "CineAssembly.h"
+#include "CineAssemblyFactory.h"
+#include "CineAssemblyNamingTokens.h"
+#include "CineAssemblySchema.h"
 #include "Editor.h"
 #include "FileHelpers.h"
+#include "MovieScene.h"
+#include "Sections/MovieSceneSubSection.h"
+#include "Sequencer/CineAssemblySchemaSequencerUtilities.h"
 #include "Subsystems/AssetEditorSubsystem.h"
+#include "Tracks/MovieSceneSubTrack.h"
+
+DEFINE_LOG_CATEGORY_STATIC(LogCineAssemblyEditorFunctionLibrary, Log, All)
+
+namespace UE::CineAssemblyTools::Private
+{
+	/** Returns the desired asset name and path for a duplicate of SequenceToDuplicate, placed relative to ParentSequence */
+	FString GetDuplicateNameAndPath(UMovieSceneSequence* SequenceToDuplicate, UMovieSceneSequence* ParentSequence, TMap<UMovieSceneSequence*, FString>& OriginalPaths)
+	{
+		// Find the original path of the parent sequence (in case the parent sequence itself is already duplicate)
+		const FString* OriginalPathPtr = OriginalPaths.Find(ParentSequence);
+		FString OriginalParentPath = OriginalPathPtr ? *OriginalPathPtr : FPaths::GetPath(ParentSequence->GetPathName());
+		OriginalParentPath = OriginalParentPath.AppendChar('/');
+
+		// Compute the relative path from the original sequence (which is being duplicated) to its parent's path
+		FString RelativePathToParent = FPaths::GetPath(SequenceToDuplicate->GetPathName());
+		FPaths::MakePathRelativeTo(RelativePathToParent, *OriginalParentPath);
+
+		// Build the desired destination path for the duplicate sequence
+		const FString DuplicatePath = FPaths::GetPath(ParentSequence->GetPathName());
+		FString DuplicateNameAndPath = DuplicatePath / RelativePathToParent / SequenceToDuplicate->GetName();
+		FPaths::CollapseRelativeDirectories(DuplicateNameAndPath);
+
+		return DuplicateNameAndPath;
+	}
+
+	/**
+	 * Recursively duplicates all of the managed SubAssemblies of the input CineAssembly, replacing the sequence referenced in the MovieScene with the duplicate.
+	 * Note: The duplicate SubAssemblies are transient, like the duplicate top-level assembly, until the asset Factory processes them to create valid packages. 
+	 */
+	void DuplicateSubAssemblies(UCineAssembly* InAssembly, TMap<UMovieSceneSequence*, FString>& OutOriginalPaths)
+	{
+		UMovieScene* MovieScene = InAssembly->GetMovieScene();
+		if (!MovieScene)
+		{
+			return;
+		}
+
+		const bool bIsReadOnly = MovieScene->IsReadOnly();
+		MovieScene->SetReadOnly(false);
+
+		for (TObjectPtr<UMovieSceneSubSection>& SubSection : InAssembly->SubAssemblies)
+		{
+			if (!SubSection)
+			{
+				continue;
+			}
+
+			UCineAssembly* OriginalSubAssembly = Cast<UCineAssembly>(SubSection->GetSequence());
+			if (!OriginalSubAssembly || OriginalSubAssembly->GetPackage() == GetTransientPackage())
+			{
+				continue;
+			}
+
+			UCineAssembly* DuplicateSubAssembly = Cast<UCineAssembly>(StaticDuplicateObject(OriginalSubAssembly, GetTransientPackage(), NAME_None, RF_Transient));
+			if (DuplicateSubAssembly)
+			{
+				DuplicateSubAssembly->ParentAssembly = InAssembly;
+
+				FCineAssemblySchemaSequencerUtilities::ReplaceSubSequence(SubSection, DuplicateSubAssembly);
+
+				OutOriginalPaths.Add(DuplicateSubAssembly, FPaths::GetPath(OriginalSubAssembly->GetPathName()));
+
+				DuplicateSubAssemblies(DuplicateSubAssembly, OutOriginalPaths);
+			}
+		}
+
+		MovieScene->SetReadOnly(bIsReadOnly);
+	}
+
+	/** Recursively duplicates all of the unmanaged subsequences of the input Sequence, replacing the sequence referenced in the MovieScene with the duplicate. */
+	void DuplicateSubsequences(UMovieSceneSequence* ParentSequence, TMap<UMovieSceneSequence*, FString>& OriginalPaths, const FString& TopLevelAssemblyPath, EDuplicateExternalSequencePreference ExternalSequencePreference)
+	{
+		UMovieScene* MovieScene = ParentSequence->GetMovieScene();
+		if (!MovieScene)
+		{
+			return;
+		}
+
+		const bool bIsReadOnly = MovieScene->IsReadOnly();
+		MovieScene->SetReadOnly(false);
+
+		for (UMovieSceneSection* Section : MovieScene->GetAllSections())
+		{
+			if (UMovieSceneSubSection* SubSection = Cast<UMovieSceneSubSection>(Section))
+			{
+				if (UMovieSceneSequence* Subsequence = SubSection->GetSequence())
+				{
+					// Managed SubAssemblies are already handled by CreateConfiguredAssembly/CreateConfiguredSubAssemblies.
+					// Skip duplication but recurse into them — the map provides the correct source path at each level.
+					UCineAssembly* ParentAssembly = Cast<UCineAssembly>(ParentSequence);
+					if (ParentAssembly && ParentAssembly->SubAssemblies.Contains(SubSection))
+					{
+						DuplicateSubsequences(Subsequence, OriginalPaths, TopLevelAssemblyPath, ExternalSequencePreference);
+						continue;
+					}
+
+					FString DesiredNameAndPath = GetDuplicateNameAndPath(Subsequence, ParentSequence, OriginalPaths);
+
+					if (!FPaths::IsUnderDirectory(DesiredNameAndPath, TopLevelAssemblyPath))
+					{
+						switch(ExternalSequencePreference)
+						{
+						case EDuplicateExternalSequencePreference::DoNotDuplicate:
+							DuplicateSubsequences(Subsequence, OriginalPaths, TopLevelAssemblyPath, ExternalSequencePreference);
+							continue;
+
+						case EDuplicateExternalSequencePreference::DuplicateIntoAssemblyFolder:
+							DesiredNameAndPath = TopLevelAssemblyPath / FPaths::GetBaseFilename(DesiredNameAndPath);
+							break;
+
+						case EDuplicateExternalSequencePreference::DuplicateIntoOriginalFolder:
+							break;
+
+						default:
+							checkNoEntry();
+						}
+					}
+
+					IAssetTools& AssetTools = FModuleManager::GetModuleChecked<FAssetToolsModule>("AssetTools").Get();
+
+					FString UniquePackageName;
+					FString UniqueAssetName;
+					AssetTools.CreateUniqueAssetName(DesiredNameAndPath, TEXT(""), UniquePackageName, UniqueAssetName);
+
+					const FString DestinationFolder = FPaths::GetPath(UniquePackageName);
+					UMovieSceneSequence* DuplicateSubsequence = Cast<UMovieSceneSequence>(AssetTools.DuplicateAsset(UniqueAssetName, DestinationFolder, Subsequence));
+
+					if (DuplicateSubsequence)
+					{
+						FCineAssemblySchemaSequencerUtilities::ReplaceSubSequence(SubSection, DuplicateSubsequence);
+
+						// Record the original's path so the recursive call can place its children correctly
+						OriginalPaths.Add(DuplicateSubsequence, FPaths::GetPath(Subsequence->GetPathName()));
+						DuplicateSubsequences(DuplicateSubsequence, OriginalPaths, TopLevelAssemblyPath, ExternalSequencePreference);
+					}
+				}
+			}
+		}
+
+		MovieScene->SetReadOnly(bIsReadOnly);
+	}
+}
 
 void UCineAssemblyEditorFunctionLibrary::OpenAssembly(UCineAssembly* CineAssembly, bool bOpenAssociatedLevel, FCineAssemblyLevelSaveOptions SaveOptions)
 {
@@ -59,3 +209,134 @@ bool UCineAssemblyEditorFunctionLibrary::OpenAssociatedLevel(UCineAssembly* Cine
 
 	return true;
 }
+
+UCineAssembly* UCineAssemblyEditorFunctionLibrary::DuplicateAssembly(UCineAssembly* SourceAssembly, const FString& DuplicatePath, const TMap<FString, FString>& Metadata, bool bDuplicateSubsequences, const FString& NameOverride, EDuplicateExternalSequencePreference ExternalSequencePreference)
+{
+	using namespace UE::CineAssemblyTools::Private;
+
+	if (!SourceAssembly)
+	{
+		return nullptr;
+	}
+
+	// If the supplied path is a virtual path, strip off the "/All" from the beginning before proceeding
+	FString CreatePath = DuplicatePath;
+	if (CreatePath.StartsWith(TEXT("/All/")))
+	{
+		CreatePath.RemoveFromStart(TEXT("/All"));
+	}
+
+	// Now, make sure that the supplied path is actually valid.
+	if (!FPackageName::IsValidPath(CreatePath))
+	{
+		UE_LOG(LogCineAssemblyEditorFunctionLibrary, Error, TEXT("Could not duplicate Cine Assembly asset because the path '%s' is not valid"), *DuplicatePath);
+		return nullptr;
+	}
+
+	// Create a transient duplicate of the source assembly to configure. The Factory will create a valid package for it later. 
+	UCineAssembly* DuplicatedAssembly = Cast<UCineAssembly>(StaticDuplicateObject(SourceAssembly, GetTransientPackage(), NAME_None, RF_Transient));
+	if (!DuplicatedAssembly)
+	{
+		return nullptr;
+	}
+
+	AddMetadataToAssembly(DuplicatedAssembly, Metadata);
+
+	// Override the assembly name (if an override was provided)
+	if (!NameOverride.IsEmpty())
+	{
+		DuplicatedAssembly->AssemblyName.Template = NameOverride;
+		DuplicatedAssembly->AssemblyName.Resolved = FText::FromString(NameOverride);
+	}
+
+	TMap<UMovieSceneSequence*, FString> OriginalPaths;
+	if (bDuplicateSubsequences)
+	{
+		// Maintain a map of the paths for each original sequence that is duplicated so that relative paths can be computed correctly even after duplicates have been made.
+		OriginalPaths.Add(DuplicatedAssembly, FPaths::GetPath(SourceAssembly->GetPathName()));
+
+		// Recursively create transient duplicates of each managed SubAssembly, which will also be given valid packages by the Factory.
+		DuplicateSubAssemblies(DuplicatedAssembly, OriginalPaths);
+	}
+
+	// Create valid packages for the duplicate assembly and all duplicate SubAssemblies
+	UCineAssemblyFactory::CreateConfiguredAssembly(DuplicatedAssembly, CreatePath);
+
+	if (bDuplicateSubsequences)
+	{
+		// Duplicate the remaining subsequences
+		DuplicateSubsequences(DuplicatedAssembly, OriginalPaths, FPaths::GetPath(DuplicatedAssembly->GetPathName()), ExternalSequencePreference);
+	}
+
+	return DuplicatedAssembly;
+}
+
+void UCineAssemblyEditorFunctionLibrary::AddMetadataToAssembly(UCineAssembly* Assembly, const TMap<FString, FString>& Metadata)
+{
+	if (!Assembly)
+	{
+		return;
+	}
+
+	const UCineAssemblySchema* Schema = Assembly->GetSchema();
+
+	for (const TPair<FString, FString>& Pair : Metadata)
+	{
+		if (Pair.Key.IsEmpty())
+		{
+			continue;
+		}
+
+		const FAssemblyMetadataDesc* MetadataDesc = Schema ? Algo::FindBy(Schema->AssemblyMetadata, Pair.Key, &FAssemblyMetadataDesc::Key) : nullptr;
+
+		if (MetadataDesc)
+		{
+			switch(MetadataDesc->Type)
+			{
+			case ECineAssemblyMetadataType::String:
+				if (MetadataDesc->bEvaluateTokens)
+				{
+					FTemplateString TemplateString;
+					TemplateString.Template = Pair.Value;
+					TemplateString.Resolved = UCineAssemblyNamingTokens::GetResolvedText(TemplateString.Template, Assembly);
+
+					Assembly->SetMetadataAsTokenString(Pair.Key, TemplateString);
+				}
+				else
+				{
+					Assembly->SetMetadataAsString(Pair.Key, Pair.Value);
+				}
+				break;
+
+			case ECineAssemblyMetadataType::Bool:
+				Assembly->SetMetadataAsBool(Pair.Key, Pair.Value.ToBool());
+				break;
+
+			case ECineAssemblyMetadataType::Integer:
+				Assembly->SetMetadataAsInteger(Pair.Key, FCString::Atoi(*Pair.Value));
+				break;
+
+			case ECineAssemblyMetadataType::Float:
+				Assembly->SetMetadataAsFloat(Pair.Key, FCString::Atof(*Pair.Value));
+				break;
+
+			case ECineAssemblyMetadataType::AssetPath:
+				Assembly->SetMetadataAsString(Pair.Key, Pair.Value);
+				break;
+
+			case ECineAssemblyMetadataType::CineAssembly:
+				Assembly->SetMetadataAsString(Pair.Key, Pair.Value);
+				break;
+
+			default:
+				checkNoEntry();
+			}
+		}
+		else
+		{
+			// If the metadata string does not match anything specified by the schema, simply add the metadata as a string.
+			Assembly->InstanceMetadata.Add(*Pair.Key, Pair.Value);
+			Assembly->SetMetadataAsString(Pair.Key, Pair.Value);
+		}
+	}
+}
diff --git a/Engine/Plugins/VirtualProduction/CinematicAssemblyTools/Source/CineAssemblyToolsEditor/Private/ProductionFunctionLibrary.cpp b/Engine/Plugins/VirtualProduction/CinematicAssemblyTools/Source/CineAssemblyToolsEditor/Private/ProductionFunctionLibrary.cpp
--- a/Engine/Plugins/VirtualProduction/CinematicAssemblyTools/Source/CineAssemblyToolsEditor/Private/ProductionFunctionLibrary.cpp
+++ b/Engine/Plugins/VirtualProduction/CinematicAssemblyTools/Source/CineAssemblyToolsEditor/Private/ProductionFunctionLibrary.cpp
@@ -4,6 +4,7 @@
 
 #include "Algo/Find.h"
 #include "CineAssembly.h"
+#include "CineAssemblyEditorFunctionLibrary.h"
 #include "CineAssemblyFactory.h"
 #include "CineAssemblyNamingTokens.h"
 #include "CineAssemblySchema.h"
@@ -109,7 +110,7 @@ bool UProductionFunctionLibrary::SetProductionExtendedData(FGuid ProductionID, c
 	return ProductionSettings->SetProductionExtendedData(ProductionID, Data);
 }
 
-UCineAssembly* UProductionFunctionLibrary::CreateAssembly(UCineAssemblySchema* Schema, TSoftObjectPtr<UWorld> Level, TSoftObjectPtr<UCineAssembly> ParentAssembly, TMap<FString, FString> Metadata, const FString& Path, const FString& Name, bool bUseDefaultNameFromSchema)
+UCineAssembly* UProductionFunctionLibrary::CreateAssembly(UCineAssemblySchema* Schema, TSoftObjectPtr<UWorld> Level, TSoftObjectPtr<UCineAssembly> ParentAssembly, const TMap<FString, FString>& Metadata, const FString& Path, const FString& Name, bool bUseDefaultNameFromSchema)
 {
 	// If the supplied path is a virtual path, strip off the "/All" from the beginning before proceeding
 	FString CreatePath = Path;
@@ -140,7 +141,7 @@ UCineAssembly* UProductionFunctionLibrary::CreateAssembly(UCineAssemblySchema* S
 		NewAssembly->ProductionName = ActiveProduction->ProductionName;
 	}
 
-	AddMetadataToAssembly(NewAssembly, Metadata);
+	UCineAssemblyEditorFunctionLibrary::AddMetadataToAssembly(NewAssembly, Metadata);
 
 	if (!Schema || !bUseDefaultNameFromSchema)
 	{
@@ -152,63 +153,3 @@ UCineAssembly* UProductionFunctionLibrary::CreateAssembly(UCineAssemblySchema* S
 
 	return NewAssembly;
 }
-
-void UProductionFunctionLibrary::AddMetadataToAssembly(UCineAssembly* Assembly, TMap<FString, FString> Metadata)
-{
-	const UCineAssemblySchema* Schema = Assembly->GetSchema();
-
-	for (const TPair<FString, FString>& Pair : Metadata)
-	{
-		if (Pair.Key.IsEmpty())
-		{
-			continue;
-		}
-
-		const FAssemblyMetadataDesc* MetadataDesc = Schema ? Algo::FindBy(Schema->AssemblyMetadata, Pair.Key, &FAssemblyMetadataDesc::Key) : nullptr;
-
-		if (MetadataDesc)
-		{
-			if (MetadataDesc->Type == ECineAssemblyMetadataType::String)
-			{
-				if (MetadataDesc->bEvaluateTokens)
-				{
-					FTemplateString TemplateString;
-					TemplateString.Template = Pair.Value;
-					TemplateString.Resolved = UCineAssemblyNamingTokens::GetResolvedText(TemplateString.Template, Assembly);
-
-					Assembly->SetMetadataAsTokenString(Pair.Key, TemplateString);
-				}
-				else
-				{
-					Assembly->SetMetadataAsString(Pair.Key, Pair.Value);
-				}
-			}
-			else if (MetadataDesc->Type == ECineAssemblyMetadataType::Bool)
-			{
-				Assembly->SetMetadataAsBool(Pair.Key, Pair.Value.ToBool());
-			}
-			else if (MetadataDesc->Type == ECineAssemblyMetadataType::Integer)
-			{
-				Assembly->SetMetadataAsInteger(Pair.Key, FCString::Atoi(*Pair.Value));
-			}
-			else if (MetadataDesc->Type == ECineAssemblyMetadataType::Float)
-			{
-				Assembly->SetMetadataAsFloat(Pair.Key, FCString::Atof(*Pair.Value));
-			}
-			else if (MetadataDesc->Type == ECineAssemblyMetadataType::AssetPath)
-			{
-				Assembly->SetMetadataAsString(Pair.Key, Pair.Value);
-			}
-			else if (MetadataDesc->Type == ECineAssemblyMetadataType::CineAssembly)
-			{
-				Assembly->SetMetadataAsString(Pair.Key, Pair.Value);
-			}
-		}
-		else
-		{
-			// If the metadata string does not match anything specified by the schema, simply add the metadata as a string.
-			Assembly->InstanceMetadata.Add(*Pair.Key, Pair.Value);
-			Assembly->SetMetadataAsString(Pair.Key, Pair.Value);
-		}
-	}
-}
diff --git a/Engine/Plugins/VirtualProduction/CinematicAssemblyTools/Source/CineAssemblyToolsEditor/Private/ProductionFunctionLibrary.h b/Engine/Plugins/VirtualProduction/CinematicAssemblyTools/Source/CineAssemblyToolsEditor/Private/ProductionFunctionLibrary.h
--- a/Engine/Plugins/VirtualProduction/CinematicAssemblyTools/Source/CineAssemblyToolsEditor/Private/ProductionFunctionLibrary.h
+++ b/Engine/Plugins/VirtualProduction/CinematicAssemblyTools/Source/CineAssemblyToolsEditor/Private/ProductionFunctionLibrary.h
@@ -76,10 +76,6 @@ public:
 	 * If bUseDefaultNameFromSchema is true, the default assembly name from the specified schema will be used as the new asset name.
 	 * It is important that any metadata required for resolving asset naming tokens is provided to this function so that the Assembly and SubAssemblies are all named correctly. 
 	 */
-	UFUNCTION(BlueprintCallable, Category = "Cine Assembly Tools")
-	static UCineAssembly* CreateAssembly(UCineAssemblySchema* Schema, TSoftObjectPtr<UWorld> Level, TSoftObjectPtr<UCineAssembly> ParentAssembly, TMap<FString, FString> Metadata, const FString& Path, const FString& Name, bool bUseDefaultNameFromSchema = true);
-
-private:
-	/** Adds the input map of metadata to the input assembly */
-	static void AddMetadataToAssembly(UCineAssembly* Assembly, TMap<FString, FString> Metadata);
+	UFUNCTION(BlueprintCallable, Category = "Cine Assembly Tools", meta = (AutoCreateRefTerm = "Metadata"))
+	static UCineAssembly* CreateAssembly(UCineAssemblySchema* Schema, TSoftObjectPtr<UWorld> Level, TSoftObjectPtr<UCineAssembly> ParentAssembly, const TMap<FString, FString>& Metadata, const FString& Path, const FString& Name, bool bUseDefaultNameFromSchema = true);
 };
diff --git a/Engine/Plugins/VirtualProduction/CinematicAssemblyTools/Source/CineAssemblyToolsEditor/Public/CineAssemblyEditorFunctionLibrary.h b/Engine/Plugins/VirtualProduction/CinematicAssemblyTools/Source/CineAssemblyToolsEditor/Public/CineAssemblyEditorFunctionLibrary.h
--- a/Engine/Plugins/VirtualProduction/CinematicAssemblyTools/Source/CineAssemblyToolsEditor/Public/CineAssemblyEditorFunctionLibrary.h
+++ b/Engine/Plugins/VirtualProduction/CinematicAssemblyTools/Source/CineAssemblyToolsEditor/Public/CineAssemblyEditorFunctionLibrary.h
@@ -25,6 +25,17 @@ struct FCineAssemblyLevelSaveOptions
 	bool bPromptForUnsavedChanges = true;
 };
 
+/** 
+ * Options for handling the duplication of subsequences that live outside of the source assembly's content tree
+ */
+UENUM(BlueprintType)
+enum class EDuplicateExternalSequencePreference : uint8
+{
+	DuplicateIntoAssemblyFolder UMETA(ToolTip = "Duplicate external subsequences into the same folder as the duplicate assembly"),
+	DuplicateIntoOriginalFolder UMETA(ToolTip = "Duplicate external subsequences into the same folder as the original subsequence"),
+	DoNotDuplicate              UMETA(ToolTip = "Do not duplicate external subsequences. Maintain the reference to the original subsequence.")
+};
+
 /**
  * Library of functions that expose CineAssembly functionality to editor scripting
  */
@@ -44,4 +55,23 @@ public:
 	 */
 	UFUNCTION(BlueprintCallable, Category = CineAssembly, meta = (AdvancedDisplay = "SaveOptions"))
 	static bool OpenAssociatedLevel(UCineAssembly* CineAssembly, FCineAssemblyLevelSaveOptions SaveOptions = FCineAssemblyLevelSaveOptions());
+
+	/**
+	 * Duplicates a CineAssembly asset at the specified path.
+	 *
+	 * @param SourceAssembly               The assembly to duplicate
+	 * @param DuplicatePath                The content browser path where the duplicate should be created
+	 * @param Metadata                     Optional metadata overrides to apply to the duplicated assembly (key/value string pairs)
+	 * @param bDuplicateSubsequences       If true, recursively duplicate all descendant subsequences. If false, subsequences will reference the same sequences as the source Assembly.
+	 * @param NameOverride                 Optional name for the duplicate. If empty, the source assembly's name is used.
+	 * @param ExternalSequencePreference   Options for handling the duplication of subsequences that live outside of the source assembly's content tree
+	 * 
+	 * @return The duplicated assembly, or nullptr if duplication failed
+	 */
+	UFUNCTION(BlueprintCallable, Category = CineAssembly, meta = (AutoCreateRefTerm = "Metadata", AdvancedDisplay = "ExternalSequencePreference"))
+	static UCineAssembly* DuplicateAssembly(UCineAssembly* SourceAssembly, const FString& DuplicatePath, const TMap<FString, FString>& Metadata, bool bDuplicateSubsequences = true, const FString& NameOverride = FString(), EDuplicateExternalSequencePreference ExternalSequencePreference = EDuplicateExternalSequencePreference::DuplicateIntoAssemblyFolder);
+
+	/** Applies each metadata key/value pair in the input map to the input Assembly */
+	UFUNCTION(BlueprintCallable, Category = CineAssembly, meta = (AutoCreateRefTerm = "Metadata"))
+	static void AddMetadataToAssembly(UCineAssembly* Assembly, const TMap<FString, FString>& Metadata);
 };
