ShotManagement: Expose some CineAssembly and CineAssemblySchema functionality to editor scripting, especially related to the Schema's template sequence and SubAssembly tracks and sections.
#rnx

#Synced-CL 51715153

diff --git a/Engine/Plugins/VirtualProduction/CinematicAssemblyTools/Source/CineAssemblyTools/Private/CineAssemblySchema.cpp b/Engine/Plugins/VirtualProduction/CinematicAssemblyTools/Source/CineAssemblyTools/Private/CineAssemblySchema.cpp
--- a/Engine/Plugins/VirtualProduction/CinematicAssemblyTools/Source/CineAssemblyTools/Private/CineAssemblySchema.cpp
+++ b/Engine/Plugins/VirtualProduction/CinematicAssemblyTools/Source/CineAssemblyTools/Private/CineAssemblySchema.cpp
@@ -2,6 +2,7 @@
 
 #include "CineAssemblySchema.h"
 
+#include "Algo/Contains.h"
 #include "AssetRegistry/AssetData.h"
 #include "AssetRegistry/AssetRegistryModule.h"
 #include "CineAssembly.h"
@@ -92,7 +93,7 @@ void UCineAssemblySchema::PostLoad()
 		{
 			// Add a new SubAssembly Track to the TemplateSequence
 			UMovieSceneSubAssemblyTrack* SubAssemblyTrack = Cast<UMovieSceneSubAssemblyTrack>(TemplateSequenceMovieScene->AddTrack(UMovieSceneSubAssemblyTrack::StaticClass()));
-			SubAssemblyTrack->SetTrackType(ESubAssemblyTrackType::SubsequenceTrack);
+			SubAssemblyTrack->TrackType = ESubAssemblyTrackType::SubsequenceTrack;
 
 			SubAssemblyTrack->SetLocalEvalDisabled(true);
 
@@ -158,6 +159,80 @@ FString UCineAssemblySchema::GetDefaultAssemblyName() const
 	return DefaultAssemblyName;
 }
 
+FString UCineAssemblySchema::GetDefaultAssemblyPath() const
+{
+	return TemplateSequence ? TemplateSequence->PathRelativeToRoot.Template : FString();
+}
+
+void UCineAssemblySchema::SetDefaultAssemblyPath(const FString& InPath)
+{
+	if (TemplateSequence)
+	{
+		TemplateSequence->Modify();
+		TemplateSequence->PathRelativeToRoot.Template = InPath;
+	}
+}
+
+TArray<FString> UCineAssemblySchema::GetDefaultFolders() const
+{
+	TArray<FString> FolderPaths;
+	if (TemplateSequence)
+	{
+		FolderPaths.Reserve(TemplateSequence->DefaultFolderNames.Num());
+		Algo::Transform(TemplateSequence->DefaultFolderNames, FolderPaths, [](const FTemplateString& FolderName) { return FolderName.Template; });
+	}
+	return FolderPaths;
+}
+
+void UCineAssemblySchema::AddDefaultFolder(const FString& FolderPath)
+{
+	if (!TemplateSequence)
+	{
+		return;
+	}
+
+	TemplateSequence->Modify();
+
+	FString NormalizedPath = FolderPath;
+	NormalizedPath.RemoveFromEnd(TEXT("/"));
+
+	// Split the path into segments and add each intermediate path (e.g. "A/B/C" adds "A", "A/B", "A/B/C")
+	TArray<FString> Segments;
+	NormalizedPath.ParseIntoArray(Segments, TEXT("/"));
+
+	FString AccumulatedPath;
+	for (const FString& Segment : Segments)
+	{
+		AccumulatedPath = AccumulatedPath.IsEmpty() ? Segment : AccumulatedPath / Segment;
+
+		if (!Algo::ContainsBy(TemplateSequence->DefaultFolderNames, AccumulatedPath, &FTemplateString::Template))
+		{
+			TemplateSequence->DefaultFolderNames.Add(FTemplateString(AccumulatedPath, FText::GetEmpty()));
+		}
+	}
+}
+
+bool UCineAssemblySchema::RemoveDefaultFolder(const FString& FolderPath)
+{
+	if (!TemplateSequence)
+	{
+		return false;
+	}
+
+	TemplateSequence->Modify();
+
+	FString NormalizedPath = FolderPath;
+	NormalizedPath.RemoveFromEnd(TEXT("/"));
+
+	// Remove the folder and all descendant folders
+	const int32 NumRemoved = TemplateSequence->DefaultFolderNames.RemoveAll([&NormalizedPath](const FTemplateString& FolderName)
+		{
+			return FolderName.Template.Equals(NormalizedPath) || FolderName.Template.StartsWith(NormalizedPath + TEXT("/"));
+		});
+
+	return NumRemoved > 0;
+}
+
 bool UCineAssemblySchema::SupportsRename() const
 {
 	return bSupportsRename;
diff --git a/Engine/Plugins/VirtualProduction/CinematicAssemblyTools/Source/CineAssemblyTools/Private/MovieSceneSubAssemblyTrack.cpp b/Engine/Plugins/VirtualProduction/CinematicAssemblyTools/Source/CineAssemblyTools/Private/MovieSceneSubAssemblyTrack.cpp
--- a/Engine/Plugins/VirtualProduction/CinematicAssemblyTools/Source/CineAssemblyTools/Private/MovieSceneSubAssemblyTrack.cpp
+++ b/Engine/Plugins/VirtualProduction/CinematicAssemblyTools/Source/CineAssemblyTools/Private/MovieSceneSubAssemblyTrack.cpp
@@ -26,11 +26,6 @@ bool UMovieSceneSubAssemblyTrack::IsCinematicShotTrack() const
 	return TrackType == ESubAssemblyTrackType::CinematicShotTrack;
 }
 
-void UMovieSceneSubAssemblyTrack::SetTrackType(ESubAssemblyTrackType InTrackType)
-{
-	TrackType = InTrackType;
-}
-
 #if WITH_EDITOR
 FText UMovieSceneSubAssemblyTrack::GetDefaultDisplayName() const
 {
diff --git a/Engine/Plugins/VirtualProduction/CinematicAssemblyTools/Source/CineAssemblyTools/Public/CineAssembly.h b/Engine/Plugins/VirtualProduction/CinematicAssemblyTools/Source/CineAssemblyTools/Public/CineAssembly.h
--- a/Engine/Plugins/VirtualProduction/CinematicAssemblyTools/Source/CineAssemblyTools/Public/CineAssembly.h
+++ b/Engine/Plugins/VirtualProduction/CinematicAssemblyTools/Source/CineAssemblyTools/Public/CineAssembly.h
@@ -62,6 +62,7 @@ public:
 	CINEASSEMBLYTOOLS_API FGuid GetAssemblyGuid() const;
 
 	/** Get the base schema for this assembly */
+	UFUNCTION(BlueprintPure, Category = "Cine Assembly Tools")
 	CINEASSEMBLYTOOLS_API const UCineAssemblySchema* GetSchema() const;
 
 	/** Set the base schema for this assembly, only if one is not already set */
diff --git a/Engine/Plugins/VirtualProduction/CinematicAssemblyTools/Source/CineAssemblyTools/Public/CineAssemblySchema.h b/Engine/Plugins/VirtualProduction/CinematicAssemblyTools/Source/CineAssemblyTools/Public/CineAssemblySchema.h
--- a/Engine/Plugins/VirtualProduction/CinematicAssemblyTools/Source/CineAssemblyTools/Public/CineAssemblySchema.h
+++ b/Engine/Plugins/VirtualProduction/CinematicAssemblyTools/Source/CineAssemblyTools/Public/CineAssemblySchema.h
@@ -76,6 +76,29 @@ public:
 	/** Get the default name for a new Assembly made from this schema */
 	CINEASSEMBLYTOOLS_API FString GetDefaultAssemblyName() const;
 
+	/** Get the default path (relative to the content root) for assemblies created from this schema */
+	UFUNCTION(BlueprintPure, Category = "Cine Assembly Tools")
+	CINEASSEMBLYTOOLS_API FString GetDefaultAssemblyPath() const;
+
+	/** Set the default path (relative to the content root) for assemblies created from this schema */
+	UFUNCTION(BlueprintCallable, Category = "Cine Assembly Tools")
+	CINEASSEMBLYTOOLS_API void SetDefaultAssemblyPath(const FString& InPath);
+
+	/** Get the list of default folder paths that will be created for assemblies made from this schema */
+	UFUNCTION(BlueprintPure, Category = "Cine Assembly Tools")
+	CINEASSEMBLYTOOLS_API TArray<FString> GetDefaultFolders() const;
+
+	/**
+	 * Add a default folder path that will be created for assemblies made from this schema.
+	 * Also adds all intermediate parent paths (e.g. adding "A/B/C" will add "A", "A/B", and "A/B/C").
+	 */
+	UFUNCTION(BlueprintCallable, Category = "Cine Assembly Tools")
+	CINEASSEMBLYTOOLS_API void AddDefaultFolder(const FString& FolderPath);
+
+	/** Remove a default folder path and all of its descendant folder paths. Returns true if any folders were found and removed. */
+	UFUNCTION(BlueprintCallable, Category = "Cine Assembly Tools")
+	CINEASSEMBLYTOOLS_API bool RemoveDefaultFolder(const FString& FolderPath);
+
 	/** Whether or not this schema can be renamed */
 	CINEASSEMBLYTOOLS_API bool SupportsRename() const;
 
@@ -86,23 +109,23 @@ public:
 
 public:
 	/** The schema name, which will be used by assemblies made from this schema as their "assembly type" */
-	UPROPERTY(EditAnywhere, Category = "Default")
+	UPROPERTY(EditAnywhere, BlueprintReadWrite, Category = "Default")
 	FString SchemaName;
 
 	/** A user-facing text description of this schema */
-	UPROPERTY(EditAnywhere, Category = "Default")
+	UPROPERTY(EditAnywhere, BlueprintReadWrite, Category = "Default")
 	FString Description;
 
 	/** The default name to be use when creating assemblies from this schema */
-	UPROPERTY(EditAnywhere, Category = "Default")
+	UPROPERTY(EditAnywhere, BlueprintReadWrite, Category = "Default")
 	FString DefaultAssemblyName;
 
 	/** Restricts assemblies made from this schema to using this Schema when picking a Parent Assembly */
-	UPROPERTY(EditAnywhere, Category = "Default", meta = (AllowedClasses = "/Script/CineAssemblyTools.CineAssemblySchema"))
+	UPROPERTY(EditAnywhere, BlueprintReadWrite, Category = "Default", meta = (AllowedClasses = "/Script/CineAssemblyTools.CineAssemblySchema"))
 	FSoftObjectPath ParentSchema;
 
 	/** The thumbnail image to use for this schema and assemblies built from this schema */
-	UPROPERTY(EditAnywhere, Category = "Default")
+	UPROPERTY(EditAnywhere, BlueprintReadWrite, Category = "Default")
 	TObjectPtr<UTexture2D> ThumbnailImage;
 
 	/** List of metadata fields that should be automatically added to assemblies made from this schema */
@@ -110,7 +133,7 @@ public:
 	TArray<FAssemblyMetadataDesc> AssemblyMetadata;
 
 	/** Template sequence, owned and managed by the schema, to use when creating the assembly as a starting point.  */
-	UPROPERTY()
+	UPROPERTY(BlueprintReadOnly, Category = "Default")
 	TObjectPtr<UCineAssembly> TemplateSequence;
 
 	/** Mark this Schema as Data-Only. Assemblies made from this Schema will have their Data-Only defaulted to this value. */
diff --git a/Engine/Plugins/VirtualProduction/CinematicAssemblyTools/Source/CineAssemblyTools/Public/MovieSceneSubAssemblySection.h b/Engine/Plugins/VirtualProduction/CinematicAssemblyTools/Source/CineAssemblyTools/Public/MovieSceneSubAssemblySection.h
--- a/Engine/Plugins/VirtualProduction/CinematicAssemblyTools/Source/CineAssemblyTools/Public/MovieSceneSubAssemblySection.h
+++ b/Engine/Plugins/VirtualProduction/CinematicAssemblyTools/Source/CineAssemblyTools/Public/MovieSceneSubAssemblySection.h
@@ -9,7 +9,7 @@
 #define UE_API CINEASSEMBLYTOOLS_API
 
 /** Indicates how a SubAssemblySection will be used */
-UENUM()
+UENUM(BlueprintType)
 enum class ESubAssemblySectionType : uint8
 {
 	Reference, /** Section contains a reference to an existing sequence */
@@ -47,10 +47,12 @@ public:
 	UE_API bool IsCinematicShotSection() const;
 
 	/** Returns the Assembly Template object */
+	UFUNCTION(BlueprintPure, Category = "Cine Assembly Tools")
 	UE_API UObject* GetAssemblyTemplate() const;
 
 	/** Sets the Assembly Template and updates the SequenceName to match the new template */
-	UE_API void SetAssemblyTemplate(UObject* TemplateObject);
+	UFUNCTION(BlueprintCallable, Category = "Cine Assembly Tools")
+	UE_API void SetAssemblyTemplate(UPARAM(meta = (AllowedClasses = "/Script/LevelSequence.LevelSequence, /Script/CineAssemblyTools.CineAssembly, /Script/CineAssemblyTools.CineAssemblySchema")) UObject* TemplateObject);
 
 	/** Get the name to use when creating a new SubAssembly for this section */
 	UE_API FText GetSequenceName() const;
@@ -66,15 +68,15 @@ public:
 
 public:
 	/** Indicates whether this section will be used as a reference to an existing sequence or as a template to create a new sequence */
-	UPROPERTY(EditAnywhere, Category = "Sequence")
+	UPROPERTY(EditAnywhere, BlueprintReadWrite, Category = "Sequence")
 	ESubAssemblySectionType SectionType = ESubAssemblySectionType::Reference;
 
 	/** The name that will be used when creating a new sequence for this section. Only applicable to Template sections */
-	UPROPERTY()
+	UPROPERTY(BlueprintReadWrite, Category = "Cine Assembly Tools")
 	FText NewSequenceName;
 
 	/** The path (relative to the root assembly's path) where the new sequence asset for this sections should be created */
-	UPROPERTY()
+	UPROPERTY(BlueprintReadWrite, Category = "Cine Assembly Tools")
 	FString NewSequencePath;
 
 	/** Name of the protected AssemblyTemplate property */
diff --git a/Engine/Plugins/VirtualProduction/CinematicAssemblyTools/Source/CineAssemblyTools/Public/MovieSceneSubAssemblyTrack.h b/Engine/Plugins/VirtualProduction/CinematicAssemblyTools/Source/CineAssemblyTools/Public/MovieSceneSubAssemblyTrack.h
--- a/Engine/Plugins/VirtualProduction/CinematicAssemblyTools/Source/CineAssemblyTools/Public/MovieSceneSubAssemblyTrack.h
+++ b/Engine/Plugins/VirtualProduction/CinematicAssemblyTools/Source/CineAssemblyTools/Public/MovieSceneSubAssemblyTrack.h
@@ -11,7 +11,7 @@
 /** 
  * Indicates the underlying track type for which this track is a placeholder
  */
-UENUM()
+UENUM(BlueprintType)
 enum class ESubAssemblyTrackType : uint8
 {
 	SubsequenceTrack,  /** This track is a placeholder for a UMovieSceneSubTrack */
@@ -40,23 +40,19 @@ public:
 #endif // WITH_EDITOR
 	//~ End UMovieSceneNameableTrack interface
 
-	/** Returns true is this track represents a UMovieSceneSubTrack */
+	/** Returns true if this track represents a UMovieSceneSubTrack */
 	UE_API bool IsSubsequenceTrack() const;
 
-	/** Returns true is this track represents a UMovieSceneCinematicShotTrack */
+	/** Returns true if this track represents a UMovieSceneCinematicShotTrack */
 	UE_API bool IsCinematicShotTrack() const;
 
-	/** Sets the underlying track type that this track is meant to represent */
-	UE_API void SetTrackType(ESubAssemblyTrackType InTrackType);
-
 #if WITH_EDITOR
 	/** Get the icon brush for this track, determined by the track type */
 	UE_API FName GetTrackIconBrush() const;
 #endif // WITH_EDITOR
 
-private:
 	/** Indicates the underlying track type for which this track is a placeholder */
-	UPROPERTY()
+	UPROPERTY(BlueprintReadWrite, Category = "Cine Assembly Tools")
 	ESubAssemblyTrackType TrackType = ESubAssemblyTrackType::CinematicShotTrack;
 };
 
diff --git a/Engine/Plugins/VirtualProduction/CinematicAssemblyTools/Source/CineAssemblyToolsEditor/Private/Sequencer/SubAssemblyTrackEditor.cpp b/Engine/Plugins/VirtualProduction/CinematicAssemblyTools/Source/CineAssemblyToolsEditor/Private/Sequencer/SubAssemblyTrackEditor.cpp
--- a/Engine/Plugins/VirtualProduction/CinematicAssemblyTools/Source/CineAssemblyToolsEditor/Private/Sequencer/SubAssemblyTrackEditor.cpp
+++ b/Engine/Plugins/VirtualProduction/CinematicAssemblyTools/Source/CineAssemblyToolsEditor/Private/Sequencer/SubAssemblyTrackEditor.cpp
@@ -194,7 +194,7 @@ UMovieSceneSubAssemblyTrack* FSubAssemblyTrackEditor::AddSubAssemblyTrack(ESubAs
 
 	// Regardless of the track type, a new MovieSceneSubAssemblyTrack is added to the Schema's template sequence
 	UMovieSceneSubAssemblyTrack* NewTrack = Cast<UMovieSceneSubAssemblyTrack>(FocusedMovieScene->AddTrack(UMovieSceneSubAssemblyTrack::StaticClass()));
-	NewTrack->SetTrackType(TrackType);
+	NewTrack->TrackType = TrackType;
 
 	if (GetSequencer().IsValid())
 	{
