ShotManagement: Add CreateSchema() to the editor scripting API.
AddSubAssemblyTemplate() and AddSubAssemblyReference() to the scripting API as convenience functions for configuring SubAssembly tracks on a Schema's template sequence.
#rnx
#Synced-CL 51715153

diff --git a/Engine/Plugins/VirtualProduction/CinematicAssemblyTools/Source/CineAssemblyToolsEditor/Private/CineAssemblyEditorFunctionLibrary.cpp b/Engine/Plugins/VirtualProduction/CinematicAssemblyTools/Source/CineAssemblyToolsEditor/Private/CineAssemblyEditorFunctionLibrary.cpp
--- a/Engine/Plugins/VirtualProduction/CinematicAssemblyTools/Source/CineAssemblyToolsEditor/Private/CineAssemblyEditorFunctionLibrary.cpp
+++ b/Engine/Plugins/VirtualProduction/CinematicAssemblyTools/Source/CineAssemblyToolsEditor/Private/CineAssemblyEditorFunctionLibrary.cpp
@@ -7,9 +7,12 @@
 #include "CineAssemblyFactory.h"
 #include "CineAssemblyNamingTokens.h"
 #include "CineAssemblySchema.h"
+#include "CineAssemblySchemaFactory.h"
 #include "Editor.h"
 #include "FileHelpers.h"
 #include "MovieScene.h"
+#include "MovieSceneSubAssemblySection.h"
+#include "MovieSceneSubAssemblyTrack.h"
 #include "Sections/MovieSceneSubSection.h"
 #include "Sequencer/CineAssemblySchemaSequencerUtilities.h"
 #include "Subsystems/AssetEditorSubsystem.h"
@@ -155,6 +158,45 @@ namespace UE::CineAssemblyTools::Private
 
 		MovieScene->SetReadOnly(bIsReadOnly);
 	}
+
+	/** Adds a new SubAssemblyTrack and Section to the template sequence of the input Schema */
+	UMovieSceneSubAssemblySection* AddSubAssemblyTrackAndSection(UCineAssemblySchema* Schema, ESubAssemblyTrackType TrackType)
+	{
+		if (!Schema)
+		{
+			UE_LOG(LogCineAssemblyEditorFunctionLibrary, Error, TEXT("Could not add SubAssembly because the Schema is null"));
+			return nullptr;
+		}
+
+		UMovieScene* MovieScene = Schema->TemplateSequence ? Schema->TemplateSequence->GetMovieScene() : nullptr;
+		if (!MovieScene)
+		{
+			UE_LOG(LogCineAssemblyEditorFunctionLibrary, Error, TEXT("Could not add SubAssembly because the Schema's TemplateSequence is invalid"));
+			return nullptr;
+		}
+
+		UMovieSceneSubAssemblyTrack* NewTrack = Cast<UMovieSceneSubAssemblyTrack>(MovieScene->AddTrack(UMovieSceneSubAssemblyTrack::StaticClass()));
+		if (!NewTrack)
+		{
+			return nullptr;
+		}
+
+		NewTrack->TrackType = TrackType;
+
+		const TRange<FFrameNumber> PlaybackRange = MovieScene->GetPlaybackRange();
+		const FFrameNumber StartFrame = PlaybackRange.GetLowerBoundValue();
+		const int32 Duration = PlaybackRange.Size<FFrameNumber>().Value;
+
+		UMovieSceneSubAssemblySection* NewSection = Cast<UMovieSceneSubAssemblySection>(NewTrack->AddSequence(nullptr, StartFrame, Duration));
+		if (!NewSection)
+		{
+			MovieScene->RemoveTrack(*NewTrack);
+			return nullptr;
+		}
+
+		Schema->MarkPackageDirty();
+		return NewSection;
+	}
 }
 
 void UCineAssemblyEditorFunctionLibrary::OpenAssembly(UCineAssembly* CineAssembly, bool bOpenAssociatedLevel, FCineAssemblyLevelSaveOptions SaveOptions)
@@ -340,3 +382,72 @@ void UCineAssemblyEditorFunctionLibrary::AddMetadataToAssembly(UCineAssembly* As
 		}
 	}
 }
+
+UCineAssemblySchema* UCineAssemblyEditorFunctionLibrary::CreateSchema(const FString& SchemaName, const FString& AssetPath, const FString& DefaultAssemblyName, const FString& RelativeAssemblyPath, const FString& Description, UCineAssemblySchema* ParentSchema, UTexture2D* ThumbnailImage, bool bIsDataOnly, bool bIsHidden)
+{
+	if (SchemaName.IsEmpty())
+	{
+		UE_LOG(LogCineAssemblyEditorFunctionLibrary, Error, TEXT("Could not create Schema because the SchemaName is empty"));
+		return nullptr;
+	}
+
+	// If the supplied path is a virtual path, strip off the "/All" from the beginning before proceeding
+	FString CreatePath = AssetPath;
+	if (CreatePath.StartsWith(TEXT("/All/")))
+	{
+		CreatePath.RemoveFromStart(TEXT("/All"));
+	}
+
+	if (!FPackageName::IsValidPath(CreatePath))
+	{
+		UE_LOG(LogCineAssemblyEditorFunctionLibrary, Error, TEXT("Could not create Schema because the path '%s' is not valid"), *AssetPath);
+		return nullptr;
+	}
+
+	// Create a transient schema and configure its properties
+	UCineAssemblySchema* NewSchema = NewObject<UCineAssemblySchema>(GetTransientPackage(), NAME_None, RF_Transient);
+	NewSchema->SchemaName = SchemaName;
+	NewSchema->DefaultAssemblyName = DefaultAssemblyName;
+	NewSchema->Description = Description;
+	NewSchema->ThumbnailImage = ThumbnailImage;
+	NewSchema->bIsDataOnly = bIsDataOnly;
+	NewSchema->bIsHidden = bIsHidden;
+	NewSchema->ParentSchema = ParentSchema ? FSoftObjectPath(ParentSchema) : FSoftObjectPath();
+	NewSchema->SetDefaultAssemblyPath(RelativeAssemblyPath);
+
+	// Persist the transient schema to a valid package
+	UCineAssemblySchemaFactory::CreateConfiguredSchema(NewSchema, CreatePath);
+
+	// Verify the schema was successfully created (i.e. it is no longer in the transient package)
+	if (NewSchema->GetPackage() == GetTransientPackage())
+	{
+		UE_LOG(LogCineAssemblyEditorFunctionLibrary, Error, TEXT("Failed to create Schema '%s' at path '%s'"), *SchemaName, *AssetPath);
+		return nullptr;
+	}
+
+	return NewSchema;
+}
+
+UMovieSceneSubAssemblySection* UCineAssemblyEditorFunctionLibrary::AddSubAssemblyTemplate(UCineAssemblySchema* Schema, ESubAssemblyTrackType TrackType, const FString& SubAssemblyName, const FString& RelativePath, UObject* TemplateObject)
+{
+	UMovieSceneSubAssemblySection* NewSection = UE::CineAssemblyTools::Private::AddSubAssemblyTrackAndSection(Schema, TrackType);
+	if (NewSection)
+	{
+		NewSection->SectionType = ESubAssemblySectionType::Template;
+		NewSection->SetAssemblyTemplate(TemplateObject);
+		NewSection->SetSequenceName(FText::AsCultureInvariant(SubAssemblyName));
+		NewSection->SetSequencePath(RelativePath);
+	}
+	return NewSection;
+}
+
+UMovieSceneSubAssemblySection* UCineAssemblyEditorFunctionLibrary::AddSubAssemblyReference(UCineAssemblySchema* Schema, ESubAssemblyTrackType TrackType, UMovieSceneSequence* Sequence)
+{
+	UMovieSceneSubAssemblySection* NewSection = UE::CineAssemblyTools::Private::AddSubAssemblyTrackAndSection(Schema, TrackType);
+	if (NewSection)
+	{
+		NewSection->SectionType = ESubAssemblySectionType::Reference;
+		NewSection->SetSequence(Sequence);
+	}
+	return NewSection;
+}
diff --git a/Engine/Plugins/VirtualProduction/CinematicAssemblyTools/Source/CineAssemblyToolsEditor/Public/CineAssemblyEditorFunctionLibrary.h b/Engine/Plugins/VirtualProduction/CinematicAssemblyTools/Source/CineAssemblyToolsEditor/Public/CineAssemblyEditorFunctionLibrary.h
--- a/Engine/Plugins/VirtualProduction/CinematicAssemblyTools/Source/CineAssemblyToolsEditor/Public/CineAssemblyEditorFunctionLibrary.h
+++ b/Engine/Plugins/VirtualProduction/CinematicAssemblyTools/Source/CineAssemblyToolsEditor/Public/CineAssemblyEditorFunctionLibrary.h
@@ -4,9 +4,14 @@
 
 #include "Kismet/BlueprintFunctionLibrary.h"
 
+#include "MovieSceneSubAssemblySection.h"
+#include "MovieSceneSubAssemblyTrack.h"
+
 #include "CineAssemblyEditorFunctionLibrary.generated.h"
 
 class UCineAssembly;
+class UCineAssemblySchema;
+class UTexture2D;
 
 /**
  * When opening a level, these options tell us how to handle unsaved changes in the currently open level
@@ -17,11 +22,11 @@ struct FCineAssemblyLevelSaveOptions
 	GENERATED_BODY()
 
 	/** Whether the currently open level should be saved before opening a different level */
-	UPROPERTY(BlueprintReadWrite, Category = CineAssembly)
+	UPROPERTY(BlueprintReadWrite, Category = "Cine Assembly Tools")
 	bool bSaveCurrentLevel = true;
 
 	/** Whether the user should be prompted before saving the current level */
-	UPROPERTY(BlueprintReadWrite, Category = CineAssembly, meta = (EditCondition = bSaveCurrentLevel))
+	UPROPERTY(BlueprintReadWrite, Category = "Cine Assembly Tools", meta = (EditCondition = bSaveCurrentLevel))
 	bool bPromptForUnsavedChanges = true;
 };
 
@@ -46,14 +51,14 @@ class UCineAssemblyEditorFunctionLibrary : public UBlueprintFunctionLibrary
 
 public:
 	/** Opens a CineAssembly in Sequencer, and optionally opens its associated Level */
-	UFUNCTION(BlueprintCallable, Category = CineAssembly, meta = (AdvancedDisplay = "SaveOptions"))
+	UFUNCTION(BlueprintCallable, Category = "Cine Assembly Tools", meta = (AdvancedDisplay = "SaveOptions"))
 	static void OpenAssembly(UCineAssembly* CineAssembly, bool bOpenAssociatedLevel = true, FCineAssemblyLevelSaveOptions SaveOptions = FCineAssemblyLevelSaveOptions());
 
 	/**
 	 * Opens the Level associated with the input CineAssembly (if it has one).
 	 * Returns false only if the input assembly was invalid, or if its associated level was valid but could not be opened.
 	 */
-	UFUNCTION(BlueprintCallable, Category = CineAssembly, meta = (AdvancedDisplay = "SaveOptions"))
+	UFUNCTION(BlueprintCallable, Category = "Cine Assembly Tools", meta = (AdvancedDisplay = "SaveOptions"))
 	static bool OpenAssociatedLevel(UCineAssembly* CineAssembly, FCineAssemblyLevelSaveOptions SaveOptions = FCineAssemblyLevelSaveOptions());
 
 	/**
@@ -68,10 +73,75 @@ public:
 	 * 
 	 * @return The duplicated assembly, or nullptr if duplication failed
 	 */
-	UFUNCTION(BlueprintCallable, Category = CineAssembly, meta = (AutoCreateRefTerm = "Metadata", AdvancedDisplay = "ExternalSequencePreference"))
+	UFUNCTION(BlueprintCallable, Category = "Cine Assembly Tools", meta = (AutoCreateRefTerm = "Metadata", AdvancedDisplay = "ExternalSequencePreference"))
 	static UCineAssembly* DuplicateAssembly(UCineAssembly* SourceAssembly, const FString& DuplicatePath, const TMap<FString, FString>& Metadata, bool bDuplicateSubsequences = true, const FString& NameOverride = FString(), EDuplicateExternalSequencePreference ExternalSequencePreference = EDuplicateExternalSequencePreference::DuplicateIntoAssemblyFolder);
 
 	/** Applies each metadata key/value pair in the input map to the input Assembly */
-	UFUNCTION(BlueprintCallable, Category = CineAssembly, meta = (AutoCreateRefTerm = "Metadata"))
+	UFUNCTION(BlueprintCallable, Category = "Cine Assembly Tools", meta = (AutoCreateRefTerm = "Metadata"))
 	static void AddMetadataToAssembly(UCineAssembly* Assembly, const TMap<FString, FString>& Metadata);
+
+	/**
+	 * Creates a new CineAssemblySchema asset at the specified path.
+	 *
+	 * @param SchemaName              The name of the schema, used as the "assembly type" for assemblies made from it
+	 * @param AssetPath               The content browser path where the schema should be created
+	 * @param DefaultAssemblyName     The default name pattern for assemblies created from this schema
+	 * @param RelativeAssemblyPath    The default path (relative to the content root) where assemblies created from this schema should be placed
+	 * @param Description             A user-facing text description of the schema
+	 * @param ParentSchema            Restricts assemblies made from this schema to only allow setting the Parent Assembly to an Assembly with this Parent Schema
+	 * @param ThumbnailImage          The thumbnail image to use for this schema and assemblies built from it
+	 * @param bIsDataOnly             Mark this schema as data-only
+	 * @param bIsHidden               Hide this schema from the UI when creating new assemblies
+	 *
+	 * @return The new schema, or nullptr if creation failed
+	 */
+	UFUNCTION(BlueprintCallable, Category = "Cine Assembly Tools", meta = (AdvancedDisplay = "Description, ParentSchema, ThumbnailImage, bIsDataOnly, bIsHidden"))
+	static UCineAssemblySchema* CreateSchema(
+		const FString& SchemaName,
+		const FString& AssetPath,
+		const FString& DefaultAssemblyName = FString(),
+		const FString& RelativeAssemblyPath = FString(),
+		const FString& Description = FString(),
+		UCineAssemblySchema* ParentSchema = nullptr,
+		UTexture2D* ThumbnailImage = nullptr,
+		bool bIsDataOnly = false,
+		bool bIsHidden = false
+	);
+
+	/**
+	 * Adds a new template SubAssembly track and section to a Schema's template sequence.
+	 * When an Assembly is created from this Schema, a new sequence will be initialized from the provided template object.
+	 *
+	 * @param Schema             The schema whose template sequence will be modified
+	 * @param TrackType          Whether this SubAssembly represents a Shot track or Subsequence track
+	 * @param SubAssemblyName    The name for the new SubAssembly sequence
+	 * @param RelativePath       The path (relative to the assembly root) where the new SubAssembly should be created
+	 * @param TemplateObject     The template object (UCineAssemblySchema, UCineAssembly, or ULevelSequence) to use when initializing the SubAssembly
+	 *
+	 * @return The new SubAssemblySection, or null if creation failed
+	 */
+	UFUNCTION(BlueprintCallable, Category = "Cine Assembly Tools")
+	static UMovieSceneSubAssemblySection* AddSubAssemblyTemplate(
+		UCineAssemblySchema* Schema,
+		ESubAssemblyTrackType TrackType = ESubAssemblyTrackType::SubsequenceTrack,
+		const FString& SubAssemblyName = FString(),
+		const FString& RelativePath = FString(),
+		UPARAM(meta = (AllowedClasses = "/Script/LevelSequence.LevelSequence, /Script/CineAssemblyTools.CineAssembly, /Script/CineAssemblyTools.CineAssemblySchema"))
+		UObject* TemplateObject = nullptr);
+
+	/**
+	 * Adds a new reference SubAssembly track and section to a Schema's template sequence.
+	 * When an Assembly is created from this Schema, the section will reference the provided sequence directly.
+	 *
+	 * @param Schema             The schema whose template sequence will be modified
+	 * @param TrackType          Whether this SubAssembly represents a Shot track or Subsequence track
+	 * @param Sequence           The existing sequence (UCineAssembly or ULevelSequence) to reference
+	 *
+	 * @return The new SubAssemblySection, or null if creation failed
+	 */
+	UFUNCTION(BlueprintCallable, Category = "Cine Assembly Tools")
+	static UMovieSceneSubAssemblySection* AddSubAssemblyReference(
+		UCineAssemblySchema* Schema,
+		ESubAssemblyTrackType TrackType = ESubAssemblyTrackType::SubsequenceTrack,
+		UMovieSceneSequence* Sequence = nullptr);
 };
