ShotManagement: Fix asset tree view widget to properly handle paths with trailing slashes, and to create intermediate folders in the tree view when necessary.
#rnx
#Synced-CL 51715153

diff --git a/Engine/Plugins/VirtualProduction/CinematicAssemblyTools/Source/CineAssemblyToolsEditor/Private/UI/CineAssembly/SCineAssemblyAssetTreeView.cpp b/Engine/Plugins/VirtualProduction/CinematicAssemblyTools/Source/CineAssemblyToolsEditor/Private/UI/CineAssembly/SCineAssemblyAssetTreeView.cpp
--- a/Engine/Plugins/VirtualProduction/CinematicAssemblyTools/Source/CineAssemblyToolsEditor/Private/UI/CineAssembly/SCineAssemblyAssetTreeView.cpp
+++ b/Engine/Plugins/VirtualProduction/CinematicAssemblyTools/Source/CineAssemblyToolsEditor/Private/UI/CineAssembly/SCineAssemblyAssetTreeView.cpp
@@ -283,14 +283,18 @@ namespace UE::CineAssemblyTools::Private
 	/** Returns the folder in the tree whose path matches the input path */
 	TSharedPtr<FCineAssemblyFolderItem> FindFolderByPathRecursive(const TSharedPtr<FCineAssemblyFolderItem>& FolderItem, const FString& FolderPath)
 	{
-		if (FolderItem->GetPath().Equals(FolderPath))
+		// Normalize the path by removing any trailing slash before comparing
+		FString NormalizedPath = FolderPath;
+		NormalizedPath.RemoveFromEnd(TEXT("/"));
+
+		if (FolderItem->GetPath().Equals(NormalizedPath))
 		{
 			return FolderItem;
 		}
 
 		for (const TSharedPtr<FCineAssemblyFolderItem>& ChildFolder : FolderItem->GetChildFolders())
 		{
-			if (const TSharedPtr<FCineAssemblyFolderItem>& MatchingFolder = FindFolderByPathRecursive(ChildFolder, FolderPath))
+			if (const TSharedPtr<FCineAssemblyFolderItem>& MatchingFolder = FindFolderByPathRecursive(ChildFolder, NormalizedPath))
 			{
 				return MatchingFolder;
 			}
@@ -356,6 +360,42 @@ namespace UE::CineAssemblyTools::Private
 		}
 	}
 
+	/** Ensures all intermediate folders in the given path exist in the tree, creating them if necessary. Returns the deepest folder. */
+	TSharedPtr<FCineAssemblyFolderItem> FindOrCreateFolderByPath(const TSharedPtr<FCineAssemblyFolderItem>& ContentRoot, const FString& FolderPath, UCineAssembly* Assembly)
+	{
+		if (FolderPath.IsEmpty())
+		{
+			return ContentRoot;
+		}
+
+		// Normalize the path by removing any trailing slash
+		FString NormalizedPath = FolderPath;
+		NormalizedPath.RemoveFromEnd(TEXT("/"));
+
+		// Check if the folder already exists
+		if (TSharedPtr<FCineAssemblyFolderItem> ExistingFolder = FindFolderByPathRecursive(ContentRoot, NormalizedPath))
+		{
+			return ExistingFolder;
+		}
+
+		// The folder doesn't exist. Ensure the parent path exists first (recursively), then create this folder.
+		const FString ParentPath = FPaths::GetPath(NormalizedPath);
+		TSharedPtr<FCineAssemblyFolderItem> ParentFolder = FindOrCreateFolderByPath(ContentRoot, ParentPath, Assembly);
+
+		if (!ParentFolder)
+		{
+			return nullptr;
+		}
+
+		TSharedPtr<FCineAssemblyFolderItem> NewFolder = MakeShared<FCineAssemblyFolderItem>();
+		NewFolder->SetPath(NormalizedPath);
+		NewFolder->ParentFolder = ParentFolder;
+		NewFolder->SetAssembly(Assembly);
+		ParentFolder->Children.Add(NewFolder);
+
+		return NewFolder;
+	}
+
 	/** Generates a unique folder path name, assuming the input folder will be the parent folder */
 	FString MakeUniqueFolderPath(const TSharedPtr<FCineAssemblyFolderItem>& FolderItem)
 	{
@@ -441,12 +481,13 @@ void SCineAssemblyAssetTreeView::InitializeContentTree(UCineAssembly* Assembly,
 	TArray<FString> FolderTemplates = GetFolderTemplates(Assembly);
 	FolderTemplates.Sort();
 
-	for (const FString& FolderPath : FolderTemplates)
+	for (FString FolderPath : FolderTemplates)
 	{
+		FolderPath.RemoveFromEnd(TEXT("/"));
 		const FString ParentPath = FPaths::GetPath(FolderPath);
 
 		// Walk the tree until we find an item whose path matches the parent path. The new tree item will be created as one of its children
-		TSharedPtr<FCineAssemblyFolderItem> ParentFolder = FindFolderByPathRecursive(ContentRoot, ParentPath);
+		TSharedPtr<FCineAssemblyFolderItem> ParentFolder = FindOrCreateFolderByPath(ContentRoot, ParentPath, Assembly);
 		if (ParentFolder)
 		{
 			TSharedPtr<FCineAssemblyFolderItem> NewItem = MakeShared<FCineAssemblyFolderItem>();
@@ -461,7 +502,7 @@ void SCineAssemblyAssetTreeView::InitializeContentTree(UCineAssembly* Assembly,
 	TopLevelAssembly = MakeShared<FSubAssemblyItem>();
 
 	const FString& TopLevelAssemblyPath = Assembly->PathRelativeToRoot.Template;
-	TSharedPtr<FCineAssemblyFolderItem> ParentFolder = FindFolderByPathRecursive(ContentRoot, TopLevelAssemblyPath);
+	TSharedPtr<FCineAssemblyFolderItem> ParentFolder = FindOrCreateFolderByPath(ContentRoot, TopLevelAssemblyPath, Assembly);
 	if (ParentFolder)
 	{
 		TopLevelAssembly->ParentFolder = ParentFolder;
@@ -489,8 +530,8 @@ void SCineAssemblyAssetTreeView::AddSubAssemblySectionsToTree(UCineAssembly* Ass
 
 		if (SubAssemblySection && SubAssemblySection->IsTemplateSection())
 		{
-			// Walk the tree until we find a parent folder whose path matches the path of the SubAssembly section
-			if (TSharedPtr<FCineAssemblyFolderItem> ParentFolder = UE::CineAssemblyTools::Private::FindFolderByPathRecursive(ContentRoot, SubAssemblySection->GetSequencePath()))
+			// Find or create the parent folder for this section's path, creating intermediate folders as needed
+			if (TSharedPtr<FCineAssemblyFolderItem> ParentFolder = UE::CineAssemblyTools::Private::FindOrCreateFolderByPath(ContentRoot, SubAssemblySection->GetSequencePath(), Assembly))
 			{
 				// Make a new tree item for this section and add it into the tree
 				TSharedPtr<FSubAssemblySectionItem> NewItem = MakeShared<FSubAssemblySectionItem>();
@@ -518,7 +559,7 @@ void SCineAssemblyAssetTreeView::AddSubAssembliesToTree(UCineAssembly* Assembly,
 			InitializeContentTree(SubAssembly, SubAssemblyContentRoot, SubAssemblyTreeItem);
 
 			// Reparent the items from the SubAssembly content tree into the parent assembly's content tree
-			if (TSharedPtr<FCineAssemblyFolderItem> ParentFolder = UE::CineAssemblyTools::Private::FindFolderByPathRecursive(ContentRoot, SubAssembly->PathRelativeToParent.Template))
+			if (TSharedPtr<FCineAssemblyFolderItem> ParentFolder = UE::CineAssemblyTools::Private::FindOrCreateFolderByPath(ContentRoot, SubAssembly->PathRelativeToParent.Template, Assembly))
 			{
 				for (const TSharedPtr<ICineAssemblyTreeItem>& Child : SubAssemblyContentRoot->Children)
 				{
