ShotManagement: Fix the Schema Sequencer to mute tracks and break possessable bindings when first opened, in case the Schema's template sequence was modified outside of the tool via editor scripting.
#rnx
#Synced-CL 51715153

diff --git a/Engine/Plugins/VirtualProduction/CinematicAssemblyTools/Source/CineAssemblyToolsEditor/Private/UI/CineAssembly/SCineAssemblySchemaWindow.cpp b/Engine/Plugins/VirtualProduction/CinematicAssemblyTools/Source/CineAssemblyToolsEditor/Private/UI/CineAssembly/SCineAssemblySchemaWindow.cpp
--- a/Engine/Plugins/VirtualProduction/CinematicAssemblyTools/Source/CineAssemblyToolsEditor/Private/UI/CineAssembly/SCineAssemblySchemaWindow.cpp
+++ b/Engine/Plugins/VirtualProduction/CinematicAssemblyTools/Source/CineAssemblyToolsEditor/Private/UI/CineAssembly/SCineAssemblySchemaWindow.cpp
@@ -22,7 +22,6 @@
 #include "PropertyEditorModule.h"
 #include "Sequencer/CineAssemblySchemaSequencerUtilities.h"
 #include "Sequencer/CineAssemblySchemaSpawnRegister.h"
-#include "SequencerUtilities.h"
 #include "SequencerSettings.h"
 #include "Settings/ContentBrowserSettings.h"
 #include "Styling/StyleColors.h"
@@ -128,6 +127,20 @@ TSharedRef<SWidget> SCineAssemblySchemaWindow::BuildUI()
 	Sequencer->OnActorAddedToSequencer().AddSP(this, &SCineAssemblySchemaWindow::OnActorAddedToSequencer);
 	Sequencer->OnMovieSceneDataChanged().AddSP(this, &SCineAssemblySchemaWindow::OnMovieSceneDataChanged);
 
+	// Ensure any existing tracks that should be muted are muted on initial load
+	MuteTracks();
+
+	// Unbind any possessable bindings that may exist in the template movie scene
+	UMovieScene* MovieScene = SchemaToConfigure->TemplateSequence ? SchemaToConfigure->TemplateSequence->GetMovieScene() : nullptr;
+	if (MovieScene)
+	{
+		for (int32 PossessableIndex = 0; PossessableIndex < MovieScene->GetPossessableCount(); ++PossessableIndex)
+		{
+			const FGuid& PossessableGuid = MovieScene->GetPossessable(PossessableIndex).GetGuid();
+			SchemaToConfigure->TemplateSequence->UnbindPossessableObjects(PossessableGuid);
+		}
+	}
+
 	return SNew(SBorder)
 		.BorderImage(FCineAssemblyToolsStyle::Get().GetBrush("ProductionWizard.PanelBackground"))
 		[
@@ -708,32 +721,11 @@ FReply SCineAssemblySchemaWindow::OnCancelClicked()
 
 void SCineAssemblySchemaWindow::OnMovieSceneDataChanged(EMovieSceneDataChangeType InChangeType)
 {
-	// In order to mute specific tracks, we only care about changes when a track is added or when a track value changes (which includes the user trying to unmute the track)
-	if ((InChangeType == EMovieSceneDataChangeType::MovieSceneStructureItemAdded) 
-		|| (InChangeType == EMovieSceneDataChangeType::MovieSceneStructureItemRemoved)
-		|| (InChangeType == EMovieSceneDataChangeType::TrackValueChanged))
-	{
-		if (UMovieScene* MovieScene = SchemaToConfigure->TemplateSequence->GetMovieScene())
-		{
-			// Because we do not have any additional information about what the actual MovieScene change was, we just loop through every track and force all relevant tracks to be muted
-			for (UMovieSceneTrack* Track : MovieScene->GetTracks())
-			{
-				// Mute CVar tracks and Subsequence Tracks to prevent their evaluation from affecting the world
-				if (UMovieSceneCVarTrack* CVarTrack = Cast<UMovieSceneCVarTrack>(Track))
-				{
-					CVarTrack->SetLocalEvalDisabled(true);
-				}
-				else if (UMovieSceneSubAssemblyTrack* SchemaSubTrack = Cast<UMovieSceneSubAssemblyTrack>(Track))
-				{
-					SchemaSubTrack->SetLocalEvalDisabled(true);
-				}
-			}
-		}
+	MuteTracks();
 
-		if (TreeView)
-		{
-			TreeView->Reinitialize();
-		}
+	if (TreeView)
+	{
+		TreeView->Reinitialize();
 	}
 }
 
@@ -742,14 +734,15 @@ void SCineAssemblySchemaWindow::OnActorAddedToSequencer(AActor* InActor, const F
 	// Possessables in our template sequence may not actually bind to any actors in the world
 	if (InActor && InBinding.IsValid())
 	{
-		FMovieSceneBindingProxy Proxy = FMovieSceneBindingProxy(InBinding, SchemaToConfigure->TemplateSequence);
-
 		// There are likely other delegates that are bound to OnActorAddedToSequencer, but we cannot control the order in which they execute.
 		// It might be dangerous to remove the actor from the newly created binding too early, before other delegates have a chance to do something with that binding.
 		// Therefore, we delay our action until the next tick to ensure that this do execute AFTER any other bound delegates. 
-		GEditor->GetTimerManager()->SetTimerForNextTick(FTimerDelegate::CreateLambda([this, Proxy]()
+		GEditor->GetTimerManager()->SetTimerForNextTick(FTimerDelegate::CreateLambda([this, InBinding]()
 			{
-				FSequencerUtilities::RemoveAllBindings(Sequencer.ToSharedRef(), Proxy);
+				if (SchemaToConfigure && SchemaToConfigure->TemplateSequence)
+				{
+					SchemaToConfigure->TemplateSequence->UnbindPossessableObjects(InBinding);
+				}
 			}));
 	}
 }
@@ -815,4 +808,22 @@ EVisibility SCineAssemblySchemaWindow::GetDisplayRateWarningVisibility() const
 	return EVisibility::Collapsed;
 }
 
+void SCineAssemblySchemaWindow::MuteTracks()
+{
+	UMovieScene* MovieScene = SchemaToConfigure->TemplateSequence ? SchemaToConfigure->TemplateSequence->GetMovieScene() : nullptr;
+	if (!MovieScene)
+	{
+		return;
+	}
+
+	for (UMovieSceneTrack* Track : MovieScene->GetTracks())
+	{
+		// Mute relevant track types to prevent their evaluation from affecting the world
+		if (Track->IsA<UMovieSceneCVarTrack>() || Track->IsA<UMovieSceneSubAssemblyTrack>())
+		{
+			Track->SetLocalEvalDisabled(true);
+		}
+	}
+}
+
 #undef LOCTEXT_NAMESPACE
diff --git a/Engine/Plugins/VirtualProduction/CinematicAssemblyTools/Source/CineAssemblyToolsEditor/Private/UI/CineAssembly/SCineAssemblySchemaWindow.h b/Engine/Plugins/VirtualProduction/CinematicAssemblyTools/Source/CineAssemblyToolsEditor/Private/UI/CineAssembly/SCineAssemblySchemaWindow.h
--- a/Engine/Plugins/VirtualProduction/CinematicAssemblyTools/Source/CineAssemblyToolsEditor/Private/UI/CineAssembly/SCineAssemblySchemaWindow.h
+++ b/Engine/Plugins/VirtualProduction/CinematicAssemblyTools/Source/CineAssemblyToolsEditor/Private/UI/CineAssembly/SCineAssemblySchemaWindow.h
@@ -107,9 +107,6 @@ private:
 	/** Callback whenever MovieScene data in the template sequence is changed */
 	void OnMovieSceneDataChanged(EMovieSceneDataChangeType InChangeType);
 
-	/** Add/delete asset items from the content hierarchy tree based on the SubAssembly sections in the TemplateSequence MovieScene */
-	void UpdateContentHierarchyAssets();
-
 	/** Extends the sequencer widget's toolbar */
 	void AddSequencerToolbarExtension(FToolBarBuilder& Builder);
 
@@ -119,6 +116,9 @@ private:
 	/** Returns the visibility of the display rate warning toolbar widget */
 	EVisibility GetDisplayRateWarningVisibility() const;
 
+	/** Mutes all tracks in the template sequence that we do not want to evaluate locally in this Sequencer widget */
+	void MuteTracks();
+
 	/** Cached content browser settings, used to restore defaults when closing the window */
 	bool bShowEngineContentCached = false;
 	bool bShowPluginContentCached = false;
