ShotManagement: Fix bug where the schema asset is not properly dirtied when modifying some subassembly section properties.
#rnx
#Synced-CL 51715153

diff --git a/Engine/Plugins/VirtualProduction/CinematicAssemblyTools/Source/CineAssemblyToolsEditor/Private/Sequencer/MovieSceneSubAssemblySectionCustomization.cpp b/Engine/Plugins/VirtualProduction/CinematicAssemblyTools/Source/CineAssemblyToolsEditor/Private/Sequencer/MovieSceneSubAssemblySectionCustomization.cpp
--- a/Engine/Plugins/VirtualProduction/CinematicAssemblyTools/Source/CineAssemblyToolsEditor/Private/Sequencer/MovieSceneSubAssemblySectionCustomization.cpp
+++ b/Engine/Plugins/VirtualProduction/CinematicAssemblyTools/Source/CineAssemblyToolsEditor/Private/Sequencer/MovieSceneSubAssemblySectionCustomization.cpp
@@ -80,6 +80,7 @@ void FSubAssemblySectionDetailCustomization::CustomizeDetails(IDetailLayoutBuild
 					})
 				.OnObjectChanged_Lambda([SubAssemblySection](const FAssetData& InAssetData)
 					{
+						SubAssemblySection->Modify();
 						SubAssemblySection->SetAssemblyTemplate(InAssetData.GetAsset());
 					})
 		];
diff --git a/Engine/Plugins/VirtualProduction/CinematicAssemblyTools/Source/CineAssemblyToolsEditor/Private/Sequencer/SubAssemblySection.cpp b/Engine/Plugins/VirtualProduction/CinematicAssemblyTools/Source/CineAssemblyToolsEditor/Private/Sequencer/SubAssemblySection.cpp
--- a/Engine/Plugins/VirtualProduction/CinematicAssemblyTools/Source/CineAssemblyToolsEditor/Private/Sequencer/SubAssemblySection.cpp
+++ b/Engine/Plugins/VirtualProduction/CinematicAssemblyTools/Source/CineAssemblyToolsEditor/Private/Sequencer/SubAssemblySection.cpp
@@ -120,6 +120,7 @@ FText FSubAssemblySection::GetSectionText() const
 
 void FSubAssemblySection::SetSectionText(const FText& InText, ETextCommit::Type InCommitType)
 {
+	MovieSceneSubAssemblySection->Modify();
 	if (MovieSceneSubAssemblySection->IsTemplateSection())
 	{
 		MovieSceneSubAssemblySection->SetSequenceName(InText);
