ShotManagement: Fix inconsistencies between assembly creation and duplication workflows.
Refactored the subsequence duplication so that the same code paths are used when duplicating an assembly and when using an assembly as a template object to initialize a new SubAssembly in a Schema template sequence.
#rnx
#Synced-CL 51715153

diff --git a/Engine/Plugins/VirtualProduction/CinematicAssemblyTools/Source/CineAssemblyTools/Private/CineAssembly.cpp b/Engine/Plugins/VirtualProduction/CinematicAssemblyTools/Source/CineAssemblyTools/Private/CineAssembly.cpp
--- a/Engine/Plugins/VirtualProduction/CinematicAssemblyTools/Source/CineAssemblyTools/Private/CineAssembly.cpp
+++ b/Engine/Plugins/VirtualProduction/CinematicAssemblyTools/Source/CineAssemblyTools/Private/CineAssembly.cpp
@@ -32,6 +32,8 @@
 
 #define LOCTEXT_NAMESPACE "CineAssembly"
 
+DEFINE_LOG_CATEGORY_STATIC(LogCineAssembly, Log, All)
+
 const FName UCineAssembly::AssetRegistryTag_AssemblyType = "AssemblyType";
 const FName UCineAssembly::AssemblyGuidPropertyName = GET_MEMBER_NAME_CHECKED(UCineAssembly, AssemblyGuid);
 
@@ -332,35 +334,91 @@ void UCineAssembly::CreateSubAssemblies()
 
 UCineAssembly* UCineAssembly::CreateSubAssemblyFromTemplate(UObject* TemplateObject)
 {
-	// The new SubAssembly begins are a transient object so that it can be configured before it is named and given a package.
+	// The new SubAssembly begins as a transient object so that it can be configured before it is named and given a package.
 	// This allows us to set its metadata first before trying to resolve any tokens which might appear in the SubAssemblyName.
-	UCineAssembly* SubAssembly = NewObject<UCineAssembly>(GetTransientPackage());
+	UCineAssembly* SubAssembly = nullptr;
 
 	if (UCineAssemblySchema* SchemaTemplate = Cast<UCineAssemblySchema>(TemplateObject))
 	{
+		SubAssembly = NewObject<UCineAssembly>(GetTransientPackage());
 		SubAssembly->InitializeFromSchema(SchemaTemplate);
 	}
 	else if (UCineAssembly* AssemblyTemplate = Cast<UCineAssembly>(TemplateObject))
 	{
-		SubAssembly->SetSchema(const_cast<UCineAssemblySchema*>(AssemblyTemplate->GetSchema()));
-		SubAssembly->InitializeFromTemplate(AssemblyTemplate);
-		SubAssembly->DuplicateSubSequences(AssemblyTemplate);
+		SubAssembly = Cast<UCineAssembly>(StaticDuplicateObject(AssemblyTemplate, GetTransientPackage(), NAME_None, RF_Transient));
+
+		if (!SubAssembly)
+		{
+			UE_LOG(LogCineAssembly, Error, TEXT("Failed to duplicate the UCineAssembly template '%s'"), *AssemblyTemplate->GetPathName());
+			return nullptr;
+		}
+
+		SubAssembly->SourceAssemblyPath = FPaths::GetPath(AssemblyTemplate->GetPathName());
+		SubAssembly->DuplicateManagedSubAssemblies();
 	}
 	else if (ULevelSequence* SequenceTemplate = Cast<ULevelSequence>(TemplateObject))
 	{
+		SubAssembly = NewObject<UCineAssembly>(GetTransientPackage());
 		SubAssembly->InitializeFromTemplate(SequenceTemplate);
-		SubAssembly->DuplicateSubSequences(SequenceTemplate);
+
+		SubAssembly->SourceAssemblyPath = FPaths::GetPath(SequenceTemplate->GetPathName());
 	}
 	else
 	{
+		// No template is specified, so create a new blank Assembly
+		SubAssembly = NewObject<UCineAssembly>(GetTransientPackage());
 		SubAssembly->Initialize();
 	}
 
 	SubAssembly->ParentAssembly = this;
-
 	return SubAssembly;
 }
 
+void UCineAssembly::DuplicateManagedSubAssemblies()
+{
+	if (!MovieScene)
+	{
+		return;
+	}
+
+#if WITH_EDITOR
+	const bool bIsReadOnly = MovieScene->IsReadOnly();
+	MovieScene->SetReadOnly(false);
+#endif // WITH_EDITOR
+
+	for (TObjectPtr<UMovieSceneSubSection>& SubSection : SubAssemblies)
+	{
+		if (!SubSection)
+		{
+			continue;
+		}
+
+		UCineAssembly* OriginalSubAssembly = Cast<UCineAssembly>(SubSection->GetSequence());
+		if (!OriginalSubAssembly || OriginalSubAssembly->GetPackage() == GetTransientPackage())
+		{
+			continue;
+		}
+
+		UCineAssembly* DuplicateSubAssembly = Cast<UCineAssembly>(StaticDuplicateObject(OriginalSubAssembly, GetTransientPackage(), NAME_None, RF_Transient));
+		if (DuplicateSubAssembly)
+		{
+			DuplicateSubAssembly->ParentAssembly = this;
+			DuplicateSubAssembly->SourceAssemblyPath = FPaths::GetPath(OriginalSubAssembly->GetPathName());
+
+			const bool bIsSectionLocked = SubSection->IsLocked();
+			SubSection->SetIsLocked(false);
+			SubSection->SetSequence(DuplicateSubAssembly);
+			SubSection->SetIsLocked(bIsSectionLocked);
+
+			DuplicateSubAssembly->DuplicateManagedSubAssemblies();
+		}
+	}
+
+#if WITH_EDITOR
+	MovieScene->SetReadOnly(bIsReadOnly);
+#endif // WITH_EDITOR
+}
+
 void UCineAssembly::ConvertSubAssemblyTracks()
 {
 	// Replace any SubAssemblyTracks we find with either a MovieSceneSubTrack or MovieSceneCinematicShotTrack, depending on the TrackType
@@ -393,6 +451,10 @@ void UCineAssembly::ConvertSubAssemblyTracks()
 				if (SubAssemblySection->IsTemplateSection())
 				{
 					UCineAssembly* NewSubAssembly = CreateSubAssemblyFromTemplate(SubAssemblySection->GetAssemblyTemplate());
+					if (!NewSubAssembly)
+					{
+						continue;
+					}
 
 					// Set the name and path of the new SubAssembly from the data stored in the SubAssembly section
 					NewSubAssembly->AssemblyName.Template = SubAssemblySection->GetSequenceName().ToString();
@@ -414,117 +476,6 @@ void UCineAssembly::ConvertSubAssemblyTracks()
 	}
 }
 
-void UCineAssembly::DuplicateSubSequences(ULevelSequence* SequenceTemplate)
-{
-	for (UMovieSceneTrack* Track : MovieScene->GetTracks())
-	{
-		if (Track->IsA<UMovieSceneSubTrack>() || Track->IsA<UMovieSceneCinematicShotTrack>())
-		{
-			for (UMovieSceneSection* Section : Track->GetAllSections())
-			{
-				if (UMovieSceneSubSection* SubSection = Cast<UMovieSceneSubSection>(Section))
-				{
-					ULevelSequence* SubSequence = Cast<ULevelSequence>(SubSection->GetSequence());
-					if (!SubSequence)
-					{
-						continue;
-					}
-
-					// If the SubSequence in this section is already a transient Assembly, skip duplicating it
-					if (SubSequence && SubSequence->IsA<UCineAssembly>() && SubSequence->GetPackage() == GetTransientPackage())
-					{
-						continue;
-					}
-
-					// Create a new SubAssembly using the existing SubSequence as a template
-					UCineAssembly* DuplicateSubAssembly = CreateSubAssemblyFromTemplate(SubSequence);
-
-					if (UCineAssembly* SubAssembly = Cast<UCineAssembly>(SubSequence))
-					{
-						DuplicateSubAssembly->AssemblyName.Template = SubAssembly->AssemblyName.Template;
-					}
-					else
-					{
-#if WITH_EDITOR
-						DuplicateSubAssembly->AssemblyName.Template = SubSequence->GetDisplayName().ToString();
-#else
-						DuplicateSubAssembly->AssemblyName.Template = SubSequence->GetName();
-#endif // WITH_EDITOR
-					}
-
-					// When the transient SubAssembly is converted into an asset, the relative path will tell the factory where to create the package relative to the root folder
-					DuplicateSubAssembly->PathRelativeToParent.Template = GetRelativePath(SubSequence, SequenceTemplate);
-
-					if (!Algo::ContainsBy(DefaultFolderNames, DuplicateSubAssembly->PathRelativeToParent.Template, &FTemplateString::Template))
-					{
-						DefaultFolderNames.Add(DuplicateSubAssembly->PathRelativeToParent);
-					}
-
-					// Replace the sequence in the section with the duplicate assembly
-					SubSection->SetSequence(DuplicateSubAssembly);
-					SubAssemblies.Add(SubSection);
-				}
-			}
-		}
-	}
-}
-
-FString UCineAssembly::GetRelativePath(ULevelSequence* SubSequence, ULevelSequence* OuterSequence)
-{
-	// If the SubSequence is an Assembly, the path we want is its content root folder. 
-	// Otherwise, we just want the actual path of the asset
-	FString SubSequencePath;
-	if (UCineAssembly* SubAssembly = Cast<UCineAssembly>(SubSequence))
-	{
-		FString AssetPath; // Unused
-		SubAssembly->GetAssetPathAndRootFolder(AssetPath, SubSequencePath);
-	}
-	else
-	{
-		SubSequencePath = FPaths::GetPath(SubSequence->GetPathName());
-	}
-
-	// If there is no outer sequence to make the path relative to, early-out
-	if (!OuterSequence)
-	{
-		return SubSequencePath;
-	}
-
-	// If the outer sequence is an Assembly, the we want is its content root folder.
-	// Otherwise, we just want the actual path of the asset
-	FString OuterSequencePath;
-	if (UCineAssembly* OuterAssembly = Cast<UCineAssembly>(OuterSequence))
-	{
-		FString OuterAssemblyAssetPath; // Unused
-		OuterAssembly->GetAssetPathAndRootFolder(OuterAssemblyAssetPath, OuterSequencePath);
-	}
-	else
-	{
-		OuterSequencePath = FPaths::GetPath(OuterSequence->GetPathName());
-	}
-
-	// Only attempt to find a relative path if the SubSequence resides within the content tree of the outer sequence
-	if (FPaths::IsUnderDirectory(SubSequencePath, OuterSequencePath))
-	{
-		// Make the path relative to the outer sequence
-		FPaths::MakePathRelativeTo(SubSequencePath, *OuterSequencePath.AppendChar('/'));
-
-		// If the outer sequence is an Assembly, check if the relative subsequence path matches any of the outer assembly's resolved default folders.
-		// If a match is found, return that folder's template name (potentially tokenized)
-		if (UCineAssembly* OuterAssembly = Cast<UCineAssembly>(OuterSequence))
-		{
-			for (const FTemplateString& DefaultPath : OuterAssembly->DefaultFolderNames)
-			{
-				if (DefaultPath.Resolved.ToString() == SubSequencePath)
-				{
-					return DefaultPath.Template;
-				}
-			}
-		}
-	}
-
-	return SubSequencePath;
-}
 
 #if WITH_EDITOR
 void UCineAssembly::CreateSubFolders()
diff --git a/Engine/Plugins/VirtualProduction/CinematicAssemblyTools/Source/CineAssemblyTools/Public/CineAssembly.h b/Engine/Plugins/VirtualProduction/CinematicAssemblyTools/Source/CineAssemblyTools/Public/CineAssembly.h
--- a/Engine/Plugins/VirtualProduction/CinematicAssemblyTools/Source/CineAssemblyTools/Public/CineAssembly.h
+++ b/Engine/Plugins/VirtualProduction/CinematicAssemblyTools/Source/CineAssemblyTools/Public/CineAssembly.h
@@ -191,6 +191,9 @@ public:
 	/** Returns the asset path and its root folder (which might be different if the schema defines a default assembly path) */
 	CINEASSEMBLYTOOLS_API void GetAssetPathAndRootFolder(FString& OutAssetPath, FString& OutRootFolder);
 
+	/** Recursively duplicates all managed SubAssemblies, replacing each with a new transient copy */
+	CINEASSEMBLYTOOLS_API void DuplicateManagedSubAssemblies();
+
 private:
 	/** Sets the base schema for this assembly and re-initializes the metadata inherited from the schema */
 	CINEASSEMBLYTOOLS_API void ChangeSchema(UCineAssemblySchema* InSchema);
@@ -204,15 +207,6 @@ private:
 	/** Creates a transient CineAssembly object using the input object as a template */
 	UCineAssembly* CreateSubAssemblyFromTemplate(UObject* TemplateObject);
 
-	/** Replace each SubSequence in this Assembly with new transient SubAssembly, using the existing SubSequence as a template */
-	void DuplicateSubSequences(ULevelSequence* SequenceTemplate = nullptr);
-
-	/**
-	 * Find the path of the input SubSequence relative to the outer sequence. If the outer sequence is a CineAssembly, then the relative path will be matched against
-	 * the folders known to belong to the Assembly's content tree, and a template (tokenized) path will be returned.
-	 */
-	FString GetRelativePath(ULevelSequence* SubSequence, ULevelSequence* OuterSequence);
-
 	/** Update the underlying json object whenever keys/values in the instance metadata map are added/removed/modified */
 	void UpdateInstanceMetadata();
 
@@ -237,6 +231,9 @@ public:
 	UPROPERTY()
 	FTemplateString PathRelativeToParent;
 
+	/** If this assembly is a duplicate of another, this represents the absolute path of the original asset */
+	FString SourceAssemblyPath;
+
 	/** The level to open before opening this asset in Sequencer */
 	UPROPERTY(EditAnywhere, Category = "Default", meta = (AllowedClasses = "/Script/Engine.World"))
 	FSoftObjectPath Level;
diff --git a/Engine/Plugins/VirtualProduction/CinematicAssemblyTools/Source/CineAssemblyToolsEditor/Private/CineAssemblyEditorFunctionLibrary.cpp b/Engine/Plugins/VirtualProduction/CinematicAssemblyTools/Source/CineAssemblyToolsEditor/Private/CineAssemblyEditorFunctionLibrary.cpp
--- a/Engine/Plugins/VirtualProduction/CinematicAssemblyTools/Source/CineAssemblyToolsEditor/Private/CineAssemblyEditorFunctionLibrary.cpp
+++ b/Engine/Plugins/VirtualProduction/CinematicAssemblyTools/Source/CineAssemblyToolsEditor/Private/CineAssemblyEditorFunctionLibrary.cpp
@@ -2,9 +2,7 @@
 
 #include "CineAssemblyEditorFunctionLibrary.h"
 
-#include "AssetToolsModule.h"
 #include "CineAssembly.h"
-#include "CineAssemblyFactory.h"
 #include "CineAssemblyNamingTokens.h"
 #include "CineAssemblySchema.h"
 #include "CineAssemblySchemaFactory.h"
@@ -13,152 +11,12 @@
 #include "MovieScene.h"
 #include "MovieSceneSubAssemblySection.h"
 #include "MovieSceneSubAssemblyTrack.h"
-#include "Sections/MovieSceneSubSection.h"
-#include "Sequencer/CineAssemblySchemaSequencerUtilities.h"
 #include "Subsystems/AssetEditorSubsystem.h"
-#include "Tracks/MovieSceneSubTrack.h"
 
 DEFINE_LOG_CATEGORY_STATIC(LogCineAssemblyEditorFunctionLibrary, Log, All)
 
 namespace UE::CineAssemblyTools::Private
 {
-	/** Returns the desired asset name and path for a duplicate of SequenceToDuplicate, placed relative to ParentSequence */
-	FString GetDuplicateNameAndPath(UMovieSceneSequence* SequenceToDuplicate, UMovieSceneSequence* ParentSequence, TMap<UMovieSceneSequence*, FString>& OriginalPaths)
-	{
-		// Find the original path of the parent sequence (in case the parent sequence itself is already duplicate)
-		const FString* OriginalPathPtr = OriginalPaths.Find(ParentSequence);
-		FString OriginalParentPath = OriginalPathPtr ? *OriginalPathPtr : FPaths::GetPath(ParentSequence->GetPathName());
-		OriginalParentPath = OriginalParentPath.AppendChar('/');
-
-		// Compute the relative path from the original sequence (which is being duplicated) to its parent's path
-		FString RelativePathToParent = FPaths::GetPath(SequenceToDuplicate->GetPathName());
-		FPaths::MakePathRelativeTo(RelativePathToParent, *OriginalParentPath);
-
-		// Build the desired destination path for the duplicate sequence
-		const FString DuplicatePath = FPaths::GetPath(ParentSequence->GetPathName());
-		FString DuplicateNameAndPath = DuplicatePath / RelativePathToParent / SequenceToDuplicate->GetName();
-		FPaths::CollapseRelativeDirectories(DuplicateNameAndPath);
-
-		return DuplicateNameAndPath;
-	}
-
-	/**
-	 * Recursively duplicates all of the managed SubAssemblies of the input CineAssembly, replacing the sequence referenced in the MovieScene with the duplicate.
-	 * Note: The duplicate SubAssemblies are transient, like the duplicate top-level assembly, until the asset Factory processes them to create valid packages. 
-	 */
-	void DuplicateSubAssemblies(UCineAssembly* InAssembly, TMap<UMovieSceneSequence*, FString>& OutOriginalPaths)
-	{
-		UMovieScene* MovieScene = InAssembly->GetMovieScene();
-		if (!MovieScene)
-		{
-			return;
-		}
-
-		const bool bIsReadOnly = MovieScene->IsReadOnly();
-		MovieScene->SetReadOnly(false);
-
-		for (TObjectPtr<UMovieSceneSubSection>& SubSection : InAssembly->SubAssemblies)
-		{
-			if (!SubSection)
-			{
-				continue;
-			}
-
-			UCineAssembly* OriginalSubAssembly = Cast<UCineAssembly>(SubSection->GetSequence());
-			if (!OriginalSubAssembly || OriginalSubAssembly->GetPackage() == GetTransientPackage())
-			{
-				continue;
-			}
-
-			UCineAssembly* DuplicateSubAssembly = Cast<UCineAssembly>(StaticDuplicateObject(OriginalSubAssembly, GetTransientPackage(), NAME_None, RF_Transient));
-			if (DuplicateSubAssembly)
-			{
-				DuplicateSubAssembly->ParentAssembly = InAssembly;
-
-				FCineAssemblySchemaSequencerUtilities::ReplaceSubSequence(SubSection, DuplicateSubAssembly);
-
-				OutOriginalPaths.Add(DuplicateSubAssembly, FPaths::GetPath(OriginalSubAssembly->GetPathName()));
-
-				DuplicateSubAssemblies(DuplicateSubAssembly, OutOriginalPaths);
-			}
-		}
-
-		MovieScene->SetReadOnly(bIsReadOnly);
-	}
-
-	/** Recursively duplicates all of the unmanaged subsequences of the input Sequence, replacing the sequence referenced in the MovieScene with the duplicate. */
-	void DuplicateSubsequences(UMovieSceneSequence* ParentSequence, TMap<UMovieSceneSequence*, FString>& OriginalPaths, const FString& TopLevelAssemblyPath, EDuplicateExternalSequencePreference ExternalSequencePreference)
-	{
-		UMovieScene* MovieScene = ParentSequence->GetMovieScene();
-		if (!MovieScene)
-		{
-			return;
-		}
-
-		const bool bIsReadOnly = MovieScene->IsReadOnly();
-		MovieScene->SetReadOnly(false);
-
-		for (UMovieSceneSection* Section : MovieScene->GetAllSections())
-		{
-			if (UMovieSceneSubSection* SubSection = Cast<UMovieSceneSubSection>(Section))
-			{
-				if (UMovieSceneSequence* Subsequence = SubSection->GetSequence())
-				{
-					// Managed SubAssemblies are already handled by CreateConfiguredAssembly/CreateConfiguredSubAssemblies.
-					// Skip duplication but recurse into them — the map provides the correct source path at each level.
-					UCineAssembly* ParentAssembly = Cast<UCineAssembly>(ParentSequence);
-					if (ParentAssembly && ParentAssembly->SubAssemblies.Contains(SubSection))
-					{
-						DuplicateSubsequences(Subsequence, OriginalPaths, TopLevelAssemblyPath, ExternalSequencePreference);
-						continue;
-					}
-
-					FString DesiredNameAndPath = GetDuplicateNameAndPath(Subsequence, ParentSequence, OriginalPaths);
-
-					if (!FPaths::IsUnderDirectory(DesiredNameAndPath, TopLevelAssemblyPath))
-					{
-						switch(ExternalSequencePreference)
-						{
-						case EDuplicateExternalSequencePreference::DoNotDuplicate:
-							DuplicateSubsequences(Subsequence, OriginalPaths, TopLevelAssemblyPath, ExternalSequencePreference);
-							continue;
-
-						case EDuplicateExternalSequencePreference::DuplicateIntoAssemblyFolder:
-							DesiredNameAndPath = TopLevelAssemblyPath / FPaths::GetBaseFilename(DesiredNameAndPath);
-							break;
-
-						case EDuplicateExternalSequencePreference::DuplicateIntoOriginalFolder:
-							break;
-
-						default:
-							checkNoEntry();
-						}
-					}
-
-					IAssetTools& AssetTools = FModuleManager::GetModuleChecked<FAssetToolsModule>("AssetTools").Get();
-
-					FString UniquePackageName;
-					FString UniqueAssetName;
-					AssetTools.CreateUniqueAssetName(DesiredNameAndPath, TEXT(""), UniquePackageName, UniqueAssetName);
-
-					const FString DestinationFolder = FPaths::GetPath(UniquePackageName);
-					UMovieSceneSequence* DuplicateSubsequence = Cast<UMovieSceneSequence>(AssetTools.DuplicateAsset(UniqueAssetName, DestinationFolder, Subsequence));
-
-					if (DuplicateSubsequence)
-					{
-						FCineAssemblySchemaSequencerUtilities::ReplaceSubSequence(SubSection, DuplicateSubsequence);
-
-						// Record the original's path so the recursive call can place its children correctly
-						OriginalPaths.Add(DuplicateSubsequence, FPaths::GetPath(Subsequence->GetPathName()));
-						DuplicateSubsequences(DuplicateSubsequence, OriginalPaths, TopLevelAssemblyPath, ExternalSequencePreference);
-					}
-				}
-			}
-		}
-
-		MovieScene->SetReadOnly(bIsReadOnly);
-	}
-
 	/** Adds a new SubAssemblyTrack and Section to the template sequence of the input Schema */
 	UMovieSceneSubAssemblySection* AddSubAssemblyTrackAndSection(UCineAssemblySchema* Schema, ESubAssemblyTrackType TrackType)
 	{
@@ -254,8 +112,6 @@ bool UCineAssemblyEditorFunctionLibrary::OpenAssociatedLevel(UCineAssembly* Cine
 
 UCineAssembly* UCineAssemblyEditorFunctionLibrary::DuplicateAssembly(UCineAssembly* SourceAssembly, const FString& DuplicatePath, const TMap<FString, FString>& Metadata, bool bDuplicateSubsequences, const FString& NameOverride, EDuplicateExternalSequencePreference ExternalSequencePreference)
 {
-	using namespace UE::CineAssemblyTools::Private;
-
 	if (!SourceAssembly)
 	{
 		return nullptr;
@@ -291,14 +147,13 @@ UCineAssembly* UCineAssemblyEditorFunctionLibrary::DuplicateAssembly(UCineAssemb
 		DuplicatedAssembly->AssemblyName.Resolved = FText::FromString(NameOverride);
 	}
 
-	TMap<UMovieSceneSequence*, FString> OriginalPaths;
 	if (bDuplicateSubsequences)
 	{
-		// Maintain a map of the paths for each original sequence that is duplicated so that relative paths can be computed correctly even after duplicates have been made.
-		OriginalPaths.Add(DuplicatedAssembly, FPaths::GetPath(SourceAssembly->GetPathName()));
+		// Set the original path so DuplicateSubsequences can compute correct relative positions
+		DuplicatedAssembly->SourceAssemblyPath = FPaths::GetPath(SourceAssembly->GetPathName());
 
 		// Recursively create transient duplicates of each managed SubAssembly, which will also be given valid packages by the Factory.
-		DuplicateSubAssemblies(DuplicatedAssembly, OriginalPaths);
+		DuplicatedAssembly->DuplicateManagedSubAssemblies();
 	}
 
 	// Create valid packages for the duplicate assembly and all duplicate SubAssemblies
@@ -307,7 +162,7 @@ UCineAssembly* UCineAssemblyEditorFunctionLibrary::DuplicateAssembly(UCineAssemb
 	if (bDuplicateSubsequences)
 	{
 		// Duplicate the remaining subsequences
-		DuplicateSubsequences(DuplicatedAssembly, OriginalPaths, FPaths::GetPath(DuplicatedAssembly->GetPathName()), ExternalSequencePreference);
+		UCineAssemblyFactory::DuplicateSubsequences(DuplicatedAssembly, ExternalSequencePreference);
 	}
 
 	return DuplicatedAssembly;
diff --git a/Engine/Plugins/VirtualProduction/CinematicAssemblyTools/Source/CineAssemblyToolsEditor/Private/CineAssemblyFactory.cpp b/Engine/Plugins/VirtualProduction/CinematicAssemblyTools/Source/CineAssemblyToolsEditor/Private/CineAssemblyFactory.cpp
--- a/Engine/Plugins/VirtualProduction/CinematicAssemblyTools/Source/CineAssemblyToolsEditor/Private/CineAssemblyFactory.cpp
+++ b/Engine/Plugins/VirtualProduction/CinematicAssemblyTools/Source/CineAssemblyToolsEditor/Private/CineAssemblyFactory.cpp
@@ -11,15 +11,121 @@
 #include "Framework/Application/SlateApplication.h"
 #include "IContentBrowserSingleton.h"
 #include "Interfaces/IMainFrameModule.h"
+#include "Misc/Paths.h"
 #include "Modules/ModuleManager.h"
 #include "MovieScene.h"
 #include "MovieSceneToolsProjectSettings.h"
 #include "Sections/MovieSceneSubSection.h"
+#include "Sequencer/CineAssemblySchemaSequencerUtilities.h"
 #include "Tracks/MovieSceneSubTrack.h"
 #include "UI/CineAssembly/SCineAssemblyConfigWindow.h"
 
 #define LOCTEXT_NAMESPACE "CineAssemblyFactory"
 
+namespace UE::CineAssemblyTools::Private
+{
+	/** Returns the desired asset name and path for a duplicate of SequenceToDuplicate, placed relative to ParentSequence */
+	FString GetDuplicateNameAndPath(UMovieSceneSequence* SequenceToDuplicate, UMovieSceneSequence* ParentSequence, const TMap<UMovieSceneSequence*, FString>& OriginalPaths)
+	{
+		// Find the original path of the parent sequence (in case the parent sequence itself is already a duplicate)
+		const FString* OriginalPathPtr = OriginalPaths.Find(ParentSequence);
+		FString OriginalParentPath = OriginalPathPtr ? *OriginalPathPtr : FPaths::GetPath(ParentSequence->GetPathName());
+		OriginalParentPath = OriginalParentPath.AppendChar('/');
+
+		// Compute the relative path from the original sequence (which is being duplicated) to its parent's path
+		FString RelativePathToParent = FPaths::GetPath(SequenceToDuplicate->GetPathName());
+		FPaths::MakePathRelativeTo(RelativePathToParent, *OriginalParentPath);
+
+		// Build the desired destination path for the duplicate sequence
+		const FString DuplicatePath = FPaths::GetPath(ParentSequence->GetPathName());
+		FString DuplicateNameAndPath = DuplicatePath / RelativePathToParent / SequenceToDuplicate->GetName();
+		FPaths::CollapseRelativeDirectories(DuplicateNameAndPath);
+
+		return DuplicateNameAndPath;
+	}
+
+	/** Recursively duplicates all of the unmanaged subsequences of the input Sequence, replacing the sequence referenced in the MovieScene with the duplicate. */
+	void DuplicateSubsequencesImpl(UMovieSceneSequence* ParentSequence, TMap<UMovieSceneSequence*, FString>& OriginalPaths, const FString& TopLevelAssemblyPath, EDuplicateExternalSequencePreference ExternalSequencePreference)
+	{
+		// If this parent is a UCineAssembly, ensure its original path is in the map so original path lookups work for its subsequences
+		if (const UCineAssembly* ParentAssembly = Cast<UCineAssembly>(ParentSequence))
+		{
+			if (!ParentAssembly->SourceAssemblyPath.IsEmpty() && !OriginalPaths.Contains(ParentSequence))
+			{
+				OriginalPaths.Add(ParentSequence, ParentAssembly->SourceAssemblyPath);
+			}
+		}
+
+		UMovieScene* MovieScene = ParentSequence->GetMovieScene();
+		if (!MovieScene)
+		{
+			return;
+		}
+
+		const bool bIsReadOnly = MovieScene->IsReadOnly();
+		MovieScene->SetReadOnly(false);
+
+		for (UMovieSceneSection* Section : MovieScene->GetAllSections())
+		{
+			if (UMovieSceneSubSection* SubSection = Cast<UMovieSceneSubSection>(Section))
+			{
+				if (UMovieSceneSequence* Subsequence = SubSection->GetSequence())
+				{
+					// Managed SubAssemblies are already handled by CreateConfiguredSubAssemblies. Skip duplication but recurse into them.
+					UCineAssembly* ParentAssembly = Cast<UCineAssembly>(ParentSequence);
+					if (ParentAssembly && ParentAssembly->SubAssemblies.Contains(SubSection))
+					{
+						DuplicateSubsequencesImpl(Subsequence, OriginalPaths, TopLevelAssemblyPath, ExternalSequencePreference);
+						continue;
+					}
+
+					FString DesiredNameAndPath = GetDuplicateNameAndPath(Subsequence, ParentSequence, OriginalPaths);
+
+					if (!FPaths::IsUnderDirectory(DesiredNameAndPath, TopLevelAssemblyPath))
+					{
+						switch (ExternalSequencePreference)
+						{
+						case EDuplicateExternalSequencePreference::DoNotDuplicate:
+							DuplicateSubsequencesImpl(Subsequence, OriginalPaths, TopLevelAssemblyPath, ExternalSequencePreference);
+							continue;
+
+						case EDuplicateExternalSequencePreference::DuplicateIntoAssemblyFolder:
+							DesiredNameAndPath = TopLevelAssemblyPath / FPaths::GetBaseFilename(DesiredNameAndPath);
+							break;
+
+						case EDuplicateExternalSequencePreference::DuplicateIntoOriginalFolder:
+							break;
+
+						default:
+							checkNoEntry();
+						}
+					}
+
+					IAssetTools& AssetTools = FModuleManager::GetModuleChecked<FAssetToolsModule>("AssetTools").Get();
+
+					FString UniquePackageName;
+					FString UniqueAssetName;
+					AssetTools.CreateUniqueAssetName(DesiredNameAndPath, TEXT(""), UniquePackageName, UniqueAssetName);
+
+					const FString DestinationFolder = FPaths::GetPath(UniquePackageName);
+					UMovieSceneSequence* DuplicateSubsequence = Cast<UMovieSceneSequence>(AssetTools.DuplicateAsset(UniqueAssetName, DestinationFolder, Subsequence));
+
+					if (DuplicateSubsequence)
+					{
+						FCineAssemblySchemaSequencerUtilities::ReplaceSubSequence(SubSection, DuplicateSubsequence);
+
+						// Record the original's path so the recursive call can place its children correctly
+						OriginalPaths.Add(DuplicateSubsequence, FPaths::GetPath(Subsequence->GetPathName()));
+						DuplicateSubsequencesImpl(DuplicateSubsequence, OriginalPaths, TopLevelAssemblyPath, ExternalSequencePreference);
+					}
+				}
+			}
+		}
+
+		MovieScene->SetReadOnly(bIsReadOnly);
+	}
+}
+
 UCineAssemblyFactory::UCineAssemblyFactory()
 {
 	SupportedClass = UCineAssembly::StaticClass();
@@ -204,4 +310,12 @@ void UCineAssemblyFactory::CreateConfiguredSubAssemblies(UCineAssembly* Configur
 	}
 }
 
+void UCineAssemblyFactory::DuplicateSubsequences(UCineAssembly* Assembly, EDuplicateExternalSequencePreference ExternalSequencePreference)
+{
+	const FString AssemblyPath = FPaths::GetPath(Assembly->GetPathName());
+
+	TMap<UMovieSceneSequence*, FString> OriginalPaths;
+	UE::CineAssemblyTools::Private::DuplicateSubsequencesImpl(Assembly, OriginalPaths, AssemblyPath, ExternalSequencePreference);
+}
+
 #undef LOCTEXT_NAMESPACE
diff --git a/Engine/Plugins/VirtualProduction/CinematicAssemblyTools/Source/CineAssemblyToolsEditor/Private/CineAssemblyFactory.h b/Engine/Plugins/VirtualProduction/CinematicAssemblyTools/Source/CineAssemblyToolsEditor/Private/CineAssemblyFactory.h
--- a/Engine/Plugins/VirtualProduction/CinematicAssemblyTools/Source/CineAssemblyToolsEditor/Private/CineAssemblyFactory.h
+++ b/Engine/Plugins/VirtualProduction/CinematicAssemblyTools/Source/CineAssemblyToolsEditor/Private/CineAssemblyFactory.h
@@ -7,6 +7,17 @@
 
 #include "CineAssemblyFactory.generated.h"
 
+/**
+ * Options for handling the duplication of subsequences that live outside of the source assembly's content tree
+ */
+UENUM(BlueprintType)
+enum class EDuplicateExternalSequencePreference : uint8
+{
+	DuplicateIntoAssemblyFolder UMETA(ToolTip = "Duplicate external subsequences into the same folder as the duplicate assembly"),
+	DuplicateIntoOriginalFolder UMETA(ToolTip = "Duplicate external subsequences into the same folder as the original subsequence"),
+	DoNotDuplicate              UMETA(ToolTip = "Do not duplicate external subsequences. Maintain the reference to the original subsequence.")
+};
+
 /** 
  * Factory class used to create new UCineAssembly objects
  * Before creating a new Cine Assembly, the factory will spawn a new window to configure the properties of the asset that is being created.
@@ -23,6 +34,12 @@ public:
 	/** Create a valid package for each SubAssembly found in a SubSequence or Shot Track of the input Assembly's MovieScene */
 	static void CreateConfiguredSubAssemblies(UCineAssembly* ConfiguredAssembly, const FString& CreateAssetPath);
 
+	/**
+	 * Recursively duplicate all subsequences of the input Assembly. This will skip over managed SubAssemblies which are handled by CreateConfiguredSubAssemblies.
+	 * Note: The input assembly should have its SourceAssemblyPath set before calling this, so that relative paths can be computed correctly for the duplicated subsequences.
+	 */
+	static void DuplicateSubsequences(UCineAssembly* Assembly, EDuplicateExternalSequencePreference ExternalSequencePreference = EDuplicateExternalSequencePreference::DuplicateIntoAssemblyFolder);
+
 	/** Takes a CineAssembly and a creation path and makes a formatted CineAssembly package name. */
 	static FString MakeAssemblyPackageName(UCineAssembly* ConfiguredAssembly, const FString& CreateAssetPath);
 
diff --git a/Engine/Plugins/VirtualProduction/CinematicAssemblyTools/Source/CineAssemblyToolsEditor/Private/ProductionFunctionLibrary.cpp b/Engine/Plugins/VirtualProduction/CinematicAssemblyTools/Source/CineAssemblyToolsEditor/Private/ProductionFunctionLibrary.cpp
--- a/Engine/Plugins/VirtualProduction/CinematicAssemblyTools/Source/CineAssemblyToolsEditor/Private/ProductionFunctionLibrary.cpp
+++ b/Engine/Plugins/VirtualProduction/CinematicAssemblyTools/Source/CineAssemblyToolsEditor/Private/ProductionFunctionLibrary.cpp
@@ -151,5 +151,7 @@ UCineAssembly* UProductionFunctionLibrary::CreateAssembly(UCineAssemblySchema* S
 
 	UCineAssemblyFactory::CreateConfiguredAssembly(NewAssembly, CreatePath);
 
+	UCineAssemblyFactory::DuplicateSubsequences(NewAssembly);
+
 	return NewAssembly;
 }
diff --git a/Engine/Plugins/VirtualProduction/CinematicAssemblyTools/Source/CineAssemblyToolsEditor/Private/UI/CineAssembly/SCineAssemblyConfigWindow.cpp b/Engine/Plugins/VirtualProduction/CinematicAssemblyTools/Source/CineAssemblyToolsEditor/Private/UI/CineAssembly/SCineAssemblyConfigWindow.cpp
--- a/Engine/Plugins/VirtualProduction/CinematicAssemblyTools/Source/CineAssemblyToolsEditor/Private/UI/CineAssembly/SCineAssemblyConfigWindow.cpp
+++ b/Engine/Plugins/VirtualProduction/CinematicAssemblyTools/Source/CineAssemblyToolsEditor/Private/UI/CineAssembly/SCineAssemblyConfigWindow.cpp
@@ -302,7 +302,7 @@ FReply SCineAssemblyConfigWindow::OnCreateAssetClicked()
 	if (ValidateAssetName())
 	{
 		UCineAssemblyFactory::CreateConfiguredAssembly(CineAssemblyToConfigure.Get(), CreateAssetPath);
-
+		UCineAssemblyFactory::DuplicateSubsequences(CineAssemblyToConfigure.Get());
 		RequestDestroyWindow();
 	}
 	return FReply::Handled();
diff --git a/Engine/Plugins/VirtualProduction/CinematicAssemblyTools/Source/CineAssemblyToolsEditor/Public/CineAssemblyEditorFunctionLibrary.h b/Engine/Plugins/VirtualProduction/CinematicAssemblyTools/Source/CineAssemblyToolsEditor/Public/CineAssemblyEditorFunctionLibrary.h
--- a/Engine/Plugins/VirtualProduction/CinematicAssemblyTools/Source/CineAssemblyToolsEditor/Public/CineAssemblyEditorFunctionLibrary.h
+++ b/Engine/Plugins/VirtualProduction/CinematicAssemblyTools/Source/CineAssemblyToolsEditor/Public/CineAssemblyEditorFunctionLibrary.h
@@ -4,6 +4,8 @@
 
 #include "Kismet/BlueprintFunctionLibrary.h"
 
+#include "CineAssemblyFactory.h"
+
 #include "MovieSceneSubAssemblySection.h"
 #include "MovieSceneSubAssemblyTrack.h"
 
@@ -30,17 +32,6 @@ struct FCineAssemblyLevelSaveOptions
 	bool bPromptForUnsavedChanges = true;
 };
 
-/** 
- * Options for handling the duplication of subsequences that live outside of the source assembly's content tree
- */
-UENUM(BlueprintType)
-enum class EDuplicateExternalSequencePreference : uint8
-{
-	DuplicateIntoAssemblyFolder UMETA(ToolTip = "Duplicate external subsequences into the same folder as the duplicate assembly"),
-	DuplicateIntoOriginalFolder UMETA(ToolTip = "Duplicate external subsequences into the same folder as the original subsequence"),
-	DoNotDuplicate              UMETA(ToolTip = "Do not duplicate external subsequences. Maintain the reference to the original subsequence.")
-};
-
 /**
  * Library of functions that expose CineAssembly functionality to editor scripting
  */
