ShotManagement: Add new feature to add new blank Levels to a schema definition. When an assembly is made from a schema that defines an additional asset, like a Level, that additional asset will also get created and associated with the Assembly.
#Synced-CL 51715153

diff --git a/Engine/Plugins/VirtualProduction/CinematicAssemblyTools/Source/CineAssemblyTools/Private/CineAssembly.cpp b/Engine/Plugins/VirtualProduction/CinematicAssemblyTools/Source/CineAssemblyTools/Private/CineAssembly.cpp
--- a/Engine/Plugins/VirtualProduction/CinematicAssemblyTools/Source/CineAssemblyTools/Private/CineAssembly.cpp
+++ b/Engine/Plugins/VirtualProduction/CinematicAssemblyTools/Source/CineAssemblyTools/Private/CineAssembly.cpp
@@ -192,6 +192,7 @@ void UCineAssembly::InitializeFromTemplate(ULevelSequence* Template)
 	{
 		PathRelativeToRoot = TemplateAssembly->PathRelativeToRoot;
 		DefaultFolderNames = TemplateAssembly->DefaultFolderNames;
+		AssociatedAssets = TemplateAssembly->AssociatedAssets;
 	}
 
 	if (UMovieScene* TemplateMovieScene = Template->GetMovieScene())
diff --git a/Engine/Plugins/VirtualProduction/CinematicAssemblyTools/Source/CineAssemblyTools/Public/CineAssembly.h b/Engine/Plugins/VirtualProduction/CinematicAssemblyTools/Source/CineAssemblyTools/Public/CineAssembly.h
--- a/Engine/Plugins/VirtualProduction/CinematicAssemblyTools/Source/CineAssemblyTools/Public/CineAssembly.h
+++ b/Engine/Plugins/VirtualProduction/CinematicAssemblyTools/Source/CineAssemblyTools/Public/CineAssembly.h
@@ -14,6 +14,35 @@ class UMovieSceneFolder;
 class UMovieSceneSubSection;
 
 /** 
+ * Defines an additional asset that should be created alongside an assembly. These definitions are configured by the Schema.
+ */
+USTRUCT()
+struct FAssemblyAssociatedAssetDesc
+{
+	GENERATED_BODY()
+
+	/** The class of asset to create */
+	UPROPERTY()
+	TSoftClassPtr<UObject> AssetClass;
+
+	/** The template name for the associated asset */
+	UPROPERTY()
+	FTemplateString AssetName;
+
+	/** The path relative to the assembly root where this asset should be created */
+	UPROPERTY()
+	FTemplateString RelativePath;
+
+	/** Reference to the asset created from this definition (populated after the asset is created) */
+	UPROPERTY()
+	TSoftObjectPtr<UObject> CreatedAsset;
+
+	/** Stable identifier for this definition */
+	UPROPERTY()
+	FGuid AssetID;
+};
+
+/**
  * A cinematic building block that associates a level sequence with a level
  */
 UCLASS(MinimalAPI, BlueprintType, HideCategories=Animation)
@@ -266,6 +295,10 @@ public:
 	UPROPERTY()
 	TArray<FTemplateString> DefaultFolderNames;
 
+	/** Array of associated asset definitions, used to create additional assets alongside this assembly */
+	UPROPERTY()
+	TArray<FAssemblyAssociatedAssetDesc> AssociatedAssets;
+
 	/** Array of Subsequence Sections created based on the schema */
 	UPROPERTY()
 	TArray<TObjectPtr<UMovieSceneSubSection>> SubAssemblies;
diff --git a/Engine/Plugins/VirtualProduction/CinematicAssemblyTools/Source/CineAssemblyToolsEditor/CineAssemblyToolsEditor.build.cs b/Engine/Plugins/VirtualProduction/CinematicAssemblyTools/Source/CineAssemblyToolsEditor/CineAssemblyToolsEditor.build.cs
--- a/Engine/Plugins/VirtualProduction/CinematicAssemblyTools/Source/CineAssemblyToolsEditor/CineAssemblyToolsEditor.build.cs
+++ b/Engine/Plugins/VirtualProduction/CinematicAssemblyTools/Source/CineAssemblyToolsEditor/CineAssemblyToolsEditor.build.cs
@@ -55,6 +55,7 @@ namespace UnrealBuildTool.Rules
 					"TakeRecorder",
 					"TakesCore",
 					"ToolMenus",
+					"ToolWidgets",
 					"UnrealEd",
 					"WorkspaceMenuStructure", 
 				}
diff --git a/Engine/Plugins/VirtualProduction/CinematicAssemblyTools/Source/CineAssemblyToolsEditor/Private/CineAssemblyEditorFunctionLibrary.cpp b/Engine/Plugins/VirtualProduction/CinematicAssemblyTools/Source/CineAssemblyToolsEditor/Private/CineAssemblyEditorFunctionLibrary.cpp
--- a/Engine/Plugins/VirtualProduction/CinematicAssemblyTools/Source/CineAssemblyToolsEditor/Private/CineAssemblyEditorFunctionLibrary.cpp
+++ b/Engine/Plugins/VirtualProduction/CinematicAssemblyTools/Source/CineAssemblyToolsEditor/Private/CineAssemblyEditorFunctionLibrary.cpp
@@ -110,7 +110,7 @@ bool UCineAssemblyEditorFunctionLibrary::OpenAssociatedLevel(UCineAssembly* Cine
 	return true;
 }
 
-UCineAssembly* UCineAssemblyEditorFunctionLibrary::DuplicateAssembly(UCineAssembly* SourceAssembly, const FString& DuplicatePath, const TMap<FString, FString>& Metadata, bool bDuplicateSubsequences, const FString& NameOverride, EDuplicateExternalSequencePreference ExternalSequencePreference)
+UCineAssembly* UCineAssemblyEditorFunctionLibrary::DuplicateAssembly(UCineAssembly* SourceAssembly, const FString& DuplicatePath, const TMap<FString, FString>& Metadata, bool bDuplicateSubsequences, bool bDuplicateAssociatedAssets, const FString& NameOverride, EDuplicateExternalSequencePreference ExternalSequencePreference)
 {
 	if (!SourceAssembly)
 	{
@@ -165,6 +165,12 @@ UCineAssembly* UCineAssemblyEditorFunctionLibrary::DuplicateAssembly(UCineAssemb
 		UCineAssemblyFactory::DuplicateSubsequences(DuplicatedAssembly, ExternalSequencePreference);
 	}
 
+	if (bDuplicateAssociatedAssets)
+	{
+		// Duplicate the associated assets of the original assembly
+		UCineAssemblyFactory::DuplicateAssociatedAssets(DuplicatedAssembly);
+	}
+
 	return DuplicatedAssembly;
 }
 
diff --git a/Engine/Plugins/VirtualProduction/CinematicAssemblyTools/Source/CineAssemblyToolsEditor/Private/CineAssemblyFactory.cpp b/Engine/Plugins/VirtualProduction/CinematicAssemblyTools/Source/CineAssemblyToolsEditor/Private/CineAssemblyFactory.cpp
--- a/Engine/Plugins/VirtualProduction/CinematicAssemblyTools/Source/CineAssemblyToolsEditor/Private/CineAssemblyFactory.cpp
+++ b/Engine/Plugins/VirtualProduction/CinematicAssemblyTools/Source/CineAssemblyToolsEditor/Private/CineAssemblyFactory.cpp
@@ -24,6 +24,20 @@
 
 namespace UE::CineAssemblyTools::Private
 {
+	/** Finds and instantiates the factory that creates new assets of the given class, or nullptr if none exists */
+	UFactory* CreateFactoryForClass(UClass* AssetClass)
+	{
+		const IAssetTools& AssetTools = FModuleManager::LoadModuleChecked<FAssetToolsModule>("AssetTools").Get();
+		for (UFactory* FactoryCDO : AssetTools.GetNewAssetFactories())
+		{
+			if (FactoryCDO->GetSupportedClass() == AssetClass)
+			{
+				return NewObject<UFactory>(GetTransientPackage(), FactoryCDO->GetClass());
+			}
+		}
+		return nullptr;
+	}
+
 	/** Returns the desired asset name and path for a duplicate of SequenceToDuplicate, placed relative to ParentSequence */
 	FString GetDuplicateNameAndPath(UMovieSceneSequence* SequenceToDuplicate, UMovieSceneSequence* ParentSequence, const TMap<UMovieSceneSequence*, FString>& OriginalPaths)
 	{
@@ -267,6 +281,9 @@ void UCineAssemblyFactory::CreateConfiguredAssembly(UCineAssembly* ConfiguredAss
 	// Create packages for each of the transient SubAssembly
 	CreateConfiguredSubAssemblies(ConfiguredAssembly, CreateAssetPath);
 
+	// Create any associated assets defined by the schema
+	CreateAssociatedAssets(ConfiguredAssembly);
+
 	// Analytics
 	UE::CineAssemblyToolsAnalytics::RecordEvent_CreateAssembly(ConfiguredAssembly);
 
@@ -318,4 +335,103 @@ void UCineAssemblyFactory::DuplicateSubsequences(UCineAssembly* Assembly, EDupli
 	UE::CineAssemblyTools::Private::DuplicateSubsequencesImpl(Assembly, OriginalPaths, AssemblyPath, ExternalSequencePreference);
 }
 
+void UCineAssemblyFactory::CreateAssociatedAssets(UCineAssembly* ConfiguredAssembly)
+{
+	if (!ConfiguredAssembly)
+	{
+		return;
+	}
+
+	FString AssetPath;
+	FString AssemblyRootFolder;
+	ConfiguredAssembly->GetAssetPathAndRootFolder(AssetPath, AssemblyRootFolder);
+
+	IAssetTools& AssetTools = FModuleManager::LoadModuleChecked<FAssetToolsModule>("AssetTools").Get();
+
+	for (FAssemblyAssociatedAssetDesc& AssetDesc : ConfiguredAssembly->AssociatedAssets)
+	{
+		// Skip definitions that already have a created asset (e.g., from a duplicated assembly)
+		if (!AssetDesc.CreatedAsset.IsNull())
+		{
+			continue;
+		}
+
+		UClass* AssetClass = AssetDesc.AssetClass.LoadSynchronous();
+		if (!AssetClass)
+		{
+			continue;
+		}
+
+		const FText ResolvedName = UCineAssemblyNamingTokens::GetResolvedText(AssetDesc.AssetName.Template, ConfiguredAssembly);
+		const FText ResolvedPath = UCineAssemblyNamingTokens::GetResolvedText(AssetDesc.RelativePath.Template, ConfiguredAssembly);
+
+		const FString DestinationPath = AssemblyRootFolder / ResolvedPath.ToString();
+
+		FString UniquePackageName;
+		FString UniqueAssetName;
+		AssetTools.CreateUniqueAssetName(DestinationPath / ResolvedName.ToString(), FString(), UniquePackageName, UniqueAssetName);
+
+		UFactory* Factory = UE::CineAssemblyTools::Private::CreateFactoryForClass(AssetClass);
+		UObject* NewAsset = AssetTools.CreateAsset(UniqueAssetName, FPaths::GetPath(UniquePackageName), AssetClass, Factory);
+
+		if (NewAsset)
+		{
+			AssetDesc.CreatedAsset = NewAsset;
+		}
+	}
+}
+
+void UCineAssemblyFactory::DuplicateAssociatedAssets(UCineAssembly* DuplicatedAssembly)
+{
+	if (!DuplicatedAssembly)
+	{
+		return;
+	}
+
+	// Only duplicate assets for assemblies that were duplicated from an existing assembly.
+	// Assemblies created fresh from a Schema have their assets created by CreateAssociatedAssets instead.
+	if (!DuplicatedAssembly->SourceAssemblyPath.IsEmpty())
+	{
+		FString AssetPath;
+		FString AssemblyRootFolder;
+		DuplicatedAssembly->GetAssetPathAndRootFolder(AssetPath, AssemblyRootFolder);
+
+		IAssetTools& AssetTools = FModuleManager::LoadModuleChecked<FAssetToolsModule>("AssetTools").Get();
+
+		for (FAssemblyAssociatedAssetDesc& AssetDesc : DuplicatedAssembly->AssociatedAssets)
+		{
+			// Only duplicate if the associated asset definition already references a real asset
+			UObject* OriginalAsset = AssetDesc.CreatedAsset.LoadSynchronous();
+			if (!OriginalAsset)
+			{
+				continue;
+			}
+
+			const FText ResolvedName = UCineAssemblyNamingTokens::GetResolvedText(AssetDesc.AssetName.Template, DuplicatedAssembly);
+			const FText ResolvedPath = UCineAssemblyNamingTokens::GetResolvedText(AssetDesc.RelativePath.Template, DuplicatedAssembly);
+
+			const FString DestinationPath = AssemblyRootFolder / ResolvedPath.ToString();
+
+			FString UniquePackageName;
+			FString UniqueAssetName;
+			AssetTools.CreateUniqueAssetName(DestinationPath / ResolvedName.ToString(), FString(), UniquePackageName, UniqueAssetName);
+
+			UObject* DuplicateAsset = AssetTools.DuplicateAsset(UniqueAssetName, FPaths::GetPath(UniquePackageName), OriginalAsset);
+			if (DuplicateAsset)
+			{
+				AssetDesc.CreatedAsset = DuplicateAsset;
+			}
+		}
+	}
+
+	// Recurse into managed SubAssemblies
+	for (UMovieSceneSubSection* SubSection : DuplicatedAssembly->SubAssemblies)
+	{
+		if (UCineAssembly* SubAssembly = SubSection ? Cast<UCineAssembly>(SubSection->GetSequence()) : nullptr)
+		{
+			DuplicateAssociatedAssets(SubAssembly);
+		}
+	}
+}
+
 #undef LOCTEXT_NAMESPACE
diff --git a/Engine/Plugins/VirtualProduction/CinematicAssemblyTools/Source/CineAssemblyToolsEditor/Private/CineAssemblyFactory.h b/Engine/Plugins/VirtualProduction/CinematicAssemblyTools/Source/CineAssemblyToolsEditor/Private/CineAssemblyFactory.h
--- a/Engine/Plugins/VirtualProduction/CinematicAssemblyTools/Source/CineAssemblyToolsEditor/Private/CineAssemblyFactory.h
+++ b/Engine/Plugins/VirtualProduction/CinematicAssemblyTools/Source/CineAssemblyToolsEditor/Private/CineAssemblyFactory.h
@@ -40,6 +40,12 @@ public:
 	 */
 	static void DuplicateSubsequences(UCineAssembly* Assembly, EDuplicateExternalSequencePreference ExternalSequencePreference = EDuplicateExternalSequencePreference::DuplicateIntoAssemblyFolder);
 
+	/** Create new assets for each associated asset definition on the input Assembly */
+	static void CreateAssociatedAssets(UCineAssembly* ConfiguredAssembly);
+
+	/** Duplicate the associated assets of a duplicated assembly, updating each definition's CreatedAsset to point to the duplicate */
+	static void DuplicateAssociatedAssets(UCineAssembly* DuplicatedAssembly);
+
 	/** Takes a CineAssembly and a creation path and makes a formatted CineAssembly package name. */
 	static FString MakeAssemblyPackageName(UCineAssembly* ConfiguredAssembly, const FString& CreateAssetPath);
 
diff --git a/Engine/Plugins/VirtualProduction/CinematicAssemblyTools/Source/CineAssemblyToolsEditor/Private/ProductionFunctionLibrary.cpp b/Engine/Plugins/VirtualProduction/CinematicAssemblyTools/Source/CineAssemblyToolsEditor/Private/ProductionFunctionLibrary.cpp
--- a/Engine/Plugins/VirtualProduction/CinematicAssemblyTools/Source/CineAssemblyToolsEditor/Private/ProductionFunctionLibrary.cpp
+++ b/Engine/Plugins/VirtualProduction/CinematicAssemblyTools/Source/CineAssemblyToolsEditor/Private/ProductionFunctionLibrary.cpp
@@ -152,6 +152,7 @@ UCineAssembly* UProductionFunctionLibrary::CreateAssembly(UCineAssemblySchema* S
 	UCineAssemblyFactory::CreateConfiguredAssembly(NewAssembly, CreatePath);
 
 	UCineAssemblyFactory::DuplicateSubsequences(NewAssembly);
+	UCineAssemblyFactory::DuplicateAssociatedAssets(NewAssembly);
 
 	return NewAssembly;
 }
diff --git a/Engine/Plugins/VirtualProduction/CinematicAssemblyTools/Source/CineAssemblyToolsEditor/Private/UI/CineAssembly/SCineAssemblyAssetTreeView.cpp b/Engine/Plugins/VirtualProduction/CinematicAssemblyTools/Source/CineAssemblyToolsEditor/Private/UI/CineAssembly/SCineAssemblyAssetTreeView.cpp
--- a/Engine/Plugins/VirtualProduction/CinematicAssemblyTools/Source/CineAssemblyToolsEditor/Private/UI/CineAssembly/SCineAssemblyAssetTreeView.cpp
+++ b/Engine/Plugins/VirtualProduction/CinematicAssemblyTools/Source/CineAssemblyToolsEditor/Private/UI/CineAssembly/SCineAssemblyAssetTreeView.cpp
@@ -6,12 +6,14 @@
 #include "CineAssembly.h"
 #include "CineAssemblyNamingTokens.h"
 #include "CineAssemblyToolsStyle.h"
+#include "Framework/Application/SlateApplication.h"
 #include "Framework/MultiBox/MultiBoxBuilder.h"
 #include "MovieScene.h"
 #include "MovieSceneCommonHelpers.h"
 #include "MovieSceneSubAssemblySection.h"
 #include "MovieSceneSubAssemblyTrack.h"
 #include "PropertyCustomizationHelpers.h"
+#include "Styling/SlateIconFinder.h"
 #include "Widgets/Layout/SBorder.h"
 #include "Widgets/SBoxPanel.h"
 #include "Widgets/Text/STextBlock.h"
@@ -25,6 +27,7 @@ struct ICineAssemblyTreeItem : TSharedFromThis<ICineAssemblyTreeItem>
 	virtual TSharedPtr<FCineAssemblyFolderItem> AsFolder() { return nullptr; }
 	virtual TSharedPtr<FSubAssemblyItem> AsSubAssembly() { return nullptr; }
 	virtual TSharedPtr<FSubAssemblySectionItem> AsSubAssemblySection() { return nullptr; }
+	virtual TSharedPtr<FAssociatedAssetItem> AsAssociatedAsset() { return nullptr; }
 
 	/** Returns true if this item should never be editable, regardless of the widget's read-only state */
 	virtual bool IsReadOnly() const { return false; }
@@ -154,6 +157,70 @@ struct FSubAssemblySectionItem : ICineAssemblyTreeItem
 	TWeakObjectPtr<UMovieSceneSubAssemblySection> SubAssemblySection;
 };
 
+struct FAssociatedAssetItem : ICineAssemblyTreeItem
+{
+	virtual TSharedPtr<FAssociatedAssetItem> AsAssociatedAsset() override { return SharedThis(this); }
+
+	/** Finds the associated asset definition by AssetID in the owning assembly's array */
+	FAssemblyAssociatedAssetDesc* FindAssetDesc() const
+	{
+		if (AssociatedAssembly.IsValid())
+		{
+			return Algo::FindBy(AssociatedAssembly->AssociatedAssets, AssetID, &FAssemblyAssociatedAssetDesc::AssetID);
+		}
+		return nullptr;
+	}
+
+	//~ Begin ICineAssemblyTreeItem Interface
+	virtual FString GetName() const override
+	{
+		if (const FAssemblyAssociatedAssetDesc* AssetDesc = FindAssetDesc())
+		{
+			return AssetDesc->AssetName.Template;
+		}
+		return FString();
+	}
+
+	virtual void SetName(const FString& InName) override
+	{
+		if (FAssemblyAssociatedAssetDesc* AssetDesc = FindAssetDesc())
+		{
+			AssetDesc->AssetName.Template = InName;
+		}
+	}
+
+	virtual FText GetDisplayName() const override
+	{
+		return FText::FromString(GetName());
+	}
+
+	virtual void SetDisplayName(const FText& InDisplayName) override
+	{
+		SetName(InDisplayName.ToString());
+	}
+
+	virtual FString GetPath() const override
+	{
+		if (const FAssemblyAssociatedAssetDesc* AssetDesc = FindAssetDesc())
+		{
+			return AssetDesc->RelativePath.Template;
+		}
+		return FString();
+	}
+
+	virtual void SetPath(const FString& InPath) override
+	{
+		if (FAssemblyAssociatedAssetDesc* AssetDesc = FindAssetDesc())
+		{
+			AssetDesc->RelativePath.Template = FPaths::GetPath(InPath);
+		}
+	}
+	//~ End ICineAssemblyTreeItem Interface
+
+	/** The stable ID of the associated asset definition */
+	FGuid AssetID;
+};
+
 struct FCineAssemblyFolderItem : ICineAssemblyTreeItem
 {
 	//~ Begin ICineAssemblyTreeItem Interface
@@ -360,6 +427,15 @@ namespace UE::CineAssemblyTools::Private
 		}
 	}
 
+	void RemoveAssociatedAsset(const TSharedPtr<FAssociatedAssetItem>& AssetItem)
+	{
+		if (UCineAssembly* Assembly = AssetItem->GetAssembly())
+		{
+			Assembly->Modify();
+			Assembly->AssociatedAssets.RemoveAll([&AssetItem](const FAssemblyAssociatedAssetDesc& AssetDesc) { return AssetDesc.AssetID == AssetItem->AssetID; });
+		}
+	}
+
 	/** Ensures all intermediate folders in the given path exist in the tree, creating them if necessary. Returns the deepest folder. */
 	TSharedPtr<FCineAssemblyFolderItem> FindOrCreateFolderByPath(const TSharedPtr<FCineAssemblyFolderItem>& ContentRoot, const FString& FolderPath, UCineAssembly* Assembly)
 	{
@@ -516,6 +592,9 @@ void SCineAssemblyAssetTreeView::InitializeContentTree(UCineAssembly* Assembly,
 	// Add an item to the tree for each child SubAssembly of this Assembly
 	AddSubAssembliesToTree(Assembly, ContentRoot);
 
+	// Add an item to the tree for each associated asset definition
+	AddAssociatedAssetsToTree(Assembly, ContentRoot);
+
 	SortTreeItemsRecursive(ContentRoot);
 }
 
@@ -571,6 +650,23 @@ void SCineAssemblyAssetTreeView::AddSubAssembliesToTree(UCineAssembly* Assembly,
 	}
 }
 
+void SCineAssemblyAssetTreeView::AddAssociatedAssetsToTree(UCineAssembly* Assembly, TSharedPtr<FCineAssemblyFolderItem>& ContentRoot)
+{
+	for (const FAssemblyAssociatedAssetDesc& AssetDesc : Assembly->AssociatedAssets)
+	{
+		const FString& RelativePath = AssetDesc.RelativePath.Template;
+
+		if (TSharedPtr<FCineAssemblyFolderItem> ParentFolder = UE::CineAssemblyTools::Private::FindOrCreateFolderByPath(ContentRoot, RelativePath, Assembly))
+		{
+			TSharedPtr<FAssociatedAssetItem> NewItem = MakeShared<FAssociatedAssetItem>();
+			NewItem->ParentFolder = ParentFolder;
+			NewItem->SetAssembly(Assembly);
+			NewItem->AssetID = AssetDesc.AssetID;
+			ParentFolder->Children.Add(NewItem);
+		}
+	}
+}
+
 void SCineAssemblyAssetTreeView::Reinitialize()
 {
 	TreeItems.Empty();
@@ -631,9 +727,18 @@ TSharedRef<ITableRow> SCineAssemblyAssetTreeView::OnGenerateTreeRow(TSharedPtr<I
 		Icon->SetImage(FCineAssemblyToolsStyle::Get().GetBrush("Icons.Folder"));
 		Icon->SetColorAndOpacity(FAppStyle::Get().GetSlateColor("ContentBrowser.DefaultFolderColor"));
 	}
+	else if (TSharedPtr<FAssociatedAssetItem> AssetItem = TreeItem->AsAssociatedAsset())
+	{
+		if (const FAssemblyAssociatedAssetDesc* AssetDesc = AssetItem->FindAssetDesc())
+		{
+			UClass* AssetClass = AssetDesc->AssetClass.LoadSynchronous();
+			Icon->SetImage(FSlateIconFinder::FindIconBrushForClass(AssetClass));
+		}
+		Icon->SetColorAndOpacity(FLinearColor::White);
+	}
 	else
 	{
-		Icon->SetImage(FCineAssemblyToolsStyle::Get().GetBrush("Icons.Sequencer"));
+		Icon->SetImage(FSlateIconFinder::FindIconBrushForClass(UCineAssembly::StaticClass()));
 		Icon->SetColorAndOpacity(FLinearColor::White);
 	}
 
@@ -946,6 +1051,10 @@ void SCineAssemblyAssetTreeView::DeleteTreeItem(TSharedPtr<ICineAssemblyTreeItem
 	{
 		UE::CineAssemblyTools::Private::RemoveSubAssemblySection(SectionItem->SubAssemblySection.Get());
 	}
+	else if (TSharedPtr<FAssociatedAssetItem> AssetItem = TreeItem->AsAssociatedAsset())
+	{
+		UE::CineAssemblyTools::Private::RemoveAssociatedAsset(AssetItem);
+	}
 
 	TreeItem->ParentFolder->Children.Remove(TreeItem);
 
@@ -976,6 +1085,10 @@ void SCineAssemblyAssetTreeView::DeleteFolderItem(const TSharedPtr<FCineAssembly
 		{
 			UE::CineAssemblyTools::Private::RemoveSubAssemblySection(ChildSection->SubAssemblySection.Get());
 		}
+		else if (TSharedPtr<FAssociatedAssetItem> ChildAsset = Child->AsAssociatedAsset())
+		{
+			UE::CineAssemblyTools::Private::RemoveAssociatedAsset(ChildAsset);
+		}
 	}
 
 	FolderItem->Children.Empty();
diff --git a/Engine/Plugins/VirtualProduction/CinematicAssemblyTools/Source/CineAssemblyToolsEditor/Private/UI/CineAssembly/SCineAssemblyAssetTreeView.h b/Engine/Plugins/VirtualProduction/CinematicAssemblyTools/Source/CineAssemblyToolsEditor/Private/UI/CineAssembly/SCineAssemblyAssetTreeView.h
--- a/Engine/Plugins/VirtualProduction/CinematicAssemblyTools/Source/CineAssemblyToolsEditor/Private/UI/CineAssembly/SCineAssemblyAssetTreeView.h
+++ b/Engine/Plugins/VirtualProduction/CinematicAssemblyTools/Source/CineAssemblyToolsEditor/Private/UI/CineAssembly/SCineAssemblyAssetTreeView.h
@@ -12,6 +12,7 @@
 struct ICineAssemblyTreeItem;
 struct FSubAssemblyItem;
 struct FSubAssemblySectionItem;
+struct FAssociatedAssetItem;
 struct FCineAssemblyFolderItem;
 
 class UMovieSceneSubAssemblySection;
@@ -90,6 +91,9 @@ private:
 	/** Add SubAssembly items to the tree from the Assembly's SubAssembly list */
 	void AddSubAssembliesToTree(UCineAssembly* Assembly, TSharedPtr<FCineAssemblyFolderItem>& ContentRoot);
 
+	/** Add associated asset items to the tree from the Assembly's AssociatedAssets list */
+	void AddAssociatedAssetsToTree(UCineAssembly* Assembly, TSharedPtr<FCineAssemblyFolderItem>& ContentRoot);
+
 	/** Generates the row widget for an entry in the tree view */
 	TSharedRef<ITableRow> OnGenerateTreeRow(TSharedPtr<ICineAssemblyTreeItem> TreeItem, const TSharedRef<STableViewBase>& OwnerTable);
 
diff --git a/Engine/Plugins/VirtualProduction/CinematicAssemblyTools/Source/CineAssemblyToolsEditor/Private/UI/CineAssembly/SCineAssemblyConfigWindow.cpp b/Engine/Plugins/VirtualProduction/CinematicAssemblyTools/Source/CineAssemblyToolsEditor/Private/UI/CineAssembly/SCineAssemblyConfigWindow.cpp
--- a/Engine/Plugins/VirtualProduction/CinematicAssemblyTools/Source/CineAssemblyToolsEditor/Private/UI/CineAssembly/SCineAssemblyConfigWindow.cpp
+++ b/Engine/Plugins/VirtualProduction/CinematicAssemblyTools/Source/CineAssemblyToolsEditor/Private/UI/CineAssembly/SCineAssemblyConfigWindow.cpp
@@ -303,6 +303,7 @@ FReply SCineAssemblyConfigWindow::OnCreateAssetClicked()
 	{
 		UCineAssemblyFactory::CreateConfiguredAssembly(CineAssemblyToConfigure.Get(), CreateAssetPath);
 		UCineAssemblyFactory::DuplicateSubsequences(CineAssemblyToConfigure.Get());
+		UCineAssemblyFactory::DuplicateAssociatedAssets(CineAssemblyToConfigure.Get());
 		RequestDestroyWindow();
 	}
 	return FReply::Handled();
diff --git a/Engine/Plugins/VirtualProduction/CinematicAssemblyTools/Source/CineAssemblyToolsEditor/Private/UI/CineAssembly/SCineAssemblySchemaWindow.cpp b/Engine/Plugins/VirtualProduction/CinematicAssemblyTools/Source/CineAssemblyToolsEditor/Private/UI/CineAssembly/SCineAssemblySchemaWindow.cpp
--- a/Engine/Plugins/VirtualProduction/CinematicAssemblyTools/Source/CineAssemblyToolsEditor/Private/UI/CineAssembly/SCineAssemblySchemaWindow.cpp
+++ b/Engine/Plugins/VirtualProduction/CinematicAssemblyTools/Source/CineAssemblyToolsEditor/Private/UI/CineAssembly/SCineAssemblySchemaWindow.cpp
@@ -3,12 +3,18 @@
 #include "SCineAssemblySchemaWindow.h"
 
 #include "Algo/Contains.h"
+#include "AssetDefinitionRegistry.h"
 #include "AssetRegistry/AssetRegistryModule.h"
+#include "AssetToolsModule.h"
 #include "AssetViewUtils.h"
+#include "CineAssemblyEditorFunctionLibrary.h"
 #include "CineAssemblyNamingTokens.h"
 #include "CineAssemblySchemaFactory.h"
 #include "CineAssemblyToolsStyle.h"
 #include "Editor.h"
+#include "Engine/World.h"
+#include "Framework/Application/SlateApplication.h"
+#include "Framework/MultiBox/MultiBoxBuilder.h"
 #include "ISequencerModule.h"
 #include "IStructureDetailsView.h"
 #include "Misc/ConfigCacheIni.h"
@@ -24,6 +30,7 @@
 #include "Sequencer/CineAssemblySchemaSpawnRegister.h"
 #include "SequencerSettings.h"
 #include "Settings/ContentBrowserSettings.h"
+#include "SPositiveActionButton.h"
 #include "Styling/StyleColors.h"
 #include "Tracks/MovieSceneCVarTrack.h"
 #include "UObject/Package.h"
@@ -33,6 +40,7 @@
 #include "Widgets/Layout/SBorder.h"
 #include "Widgets/Layout/SSeparator.h"
 #include "Widgets/Layout/SSplitter.h"
+#include "Widgets/SAssetMenuIcon.h"
 #include "Widgets/SBoxPanel.h"
 #include "Widgets/Text/STextBlock.h"
 
@@ -569,28 +577,11 @@ TSharedRef<SWidget> SCineAssemblySchemaWindow::MakeHierarchyTabContent()
 			SNew(SHorizontalBox)
 
 			+ SHorizontalBox::Slot()
+				.AutoWidth()
 				[
-					SNew(SButton)
-						.ContentPadding(FMargin(2.0f))
-						.OnClicked(this, &SCineAssemblySchemaWindow::OnAddAssetClicked)
-						[
-							SNew(SHorizontalBox)
-
-								+ SHorizontalBox::Slot()
-								.AutoWidth()
-								.Padding(0.0f, 0.0f, 4.0f, 0.0f)
-								[
-									SNew(SImage)
-										.Image(FAppStyle::Get().GetBrush("Icons.Plus"))
-										.ColorAndOpacity(FStyleColors::AccentGreen)
-								]
-
-								+ SHorizontalBox::Slot()
-								.AutoWidth()
-								[
-									SNew(STextBlock).Text(LOCTEXT("AddAssetButton", "Add Asset"))
-								]
-						]
+					SNew(SPositiveActionButton)
+						.Text(LOCTEXT("AddAssetButton", "Add Asset"))
+						.OnGetMenuContent(this, &SCineAssemblySchemaWindow::OnGetAddAssetMenuContent)
 				]
 
 				+ SHorizontalBox::Slot()
@@ -661,11 +652,92 @@ TSharedRef<SWidget> SCineAssemblySchemaWindow::MakeHierarchyTabContent()
 		];
 }
 
-FReply SCineAssemblySchemaWindow::OnAddAssetClicked()
+TSharedRef<SWidget> SCineAssemblySchemaWindow::OnGetAddAssetMenuContent()
 {
-	// TODO: Add support for adding any asset type to the schema definition (not just subsequences)
-	// TODO: When adding subsequence assets, they will also be added as new tracks to the TemplateSequence MovieScene
-	return FReply::Handled();
+	constexpr bool bCloseAfterMenuSelection = true;
+	FMenuBuilder MenuBuilder(bCloseAfterMenuSelection, nullptr);
+
+	// When adding a new Cine Assembly, add a new SubAssemblyTrack and Section to the TemplateSequence
+	AddAssetMenuRow(MenuBuilder, UCineAssembly::StaticClass(), FExecuteAction::CreateLambda([this]()
+		{
+			UCineAssemblyEditorFunctionLibrary::AddSubAssemblyTemplate(SchemaToConfigure.Get(), ESubAssemblyTrackType::SubsequenceTrack, TEXT("NewCineAssembly"));
+			Sequencer->NotifyMovieSceneDataChanged(EMovieSceneDataChangeType::MovieSceneStructureItemAdded);
+			TreeView->Reinitialize();
+		})
+	);
+
+	AddAssetMenuRow(MenuBuilder, UWorld::StaticClass());
+
+	return MenuBuilder.MakeWidget();
+}
+
+void SCineAssemblySchemaWindow::AddAssetMenuRow(FMenuBuilder& MenuBuilder, const UClass* AssetClass, FExecuteAction OnClickAction)
+{
+	if (!OnClickAction.IsBound())
+	{
+		OnClickAction = FExecuteAction::CreateLambda([this, AssetClass]() { AddAssociatedAsset(AssetClass); });
+	}
+
+
+	const UAssetDefinition* AssetDefinition = UAssetDefinitionRegistry::Get()->GetAssetDefinitionForClass(AssetClass);
+	const FText DisplayName = AssetDefinition ? AssetDefinition->GetAssetDisplayName() : AssetClass->GetDisplayNameText();
+	const FText Tooltip = FText::Format(LOCTEXT("AddAssetToSchemaTooltip", "Add a new {0} to the schema"), DisplayName);
+
+	TSharedPtr<SWidget> Widget = SNew(SHorizontalBox)
+
+		+ SHorizontalBox::Slot()
+			.AutoWidth()
+			.VAlign(VAlign_Center)
+			.Padding(0.0f, 0.0f, 8.0f, 0.0f)
+			[
+				SNew(SAssetMenuIcon, AssetClass)
+					.IconContainerSize(FVector2D(24.0f, 24.0f))
+					.IconSize(FVector2D(20.0f, 20.0f))
+			]
+
+		+ SHorizontalBox::Slot()
+			.VAlign(VAlign_Center)
+			[
+				SNew(STextBlock).Text(DisplayName)
+			];
+
+	MenuBuilder.AddMenuEntry(FUIAction(OnClickAction), Widget.ToSharedRef(), NAME_None, Tooltip);
+}
+
+void SCineAssemblySchemaWindow::AddAssociatedAsset(const UClass* AssetClass)
+{
+	if (!SchemaToConfigure->TemplateSequence)
+	{
+		return;
+	}
+
+	// Query the factory for the default asset name
+	FString DefaultName;
+
+	const IAssetTools& AssetTools = FModuleManager::LoadModuleChecked<FAssetToolsModule>("AssetTools").Get();
+	for (UFactory* Factory : AssetTools.GetNewAssetFactories())
+	{
+		if (Factory->GetSupportedClass() == AssetClass)
+		{
+			DefaultName = Factory->GetDefaultNewAssetName();
+			break;
+		}
+	}
+
+	if (DefaultName.IsEmpty())
+	{
+		DefaultName = FString::Printf(TEXT("New%s"), *AssetClass->GetName());
+	}
+
+	// Add a new associated asset to the template sequence
+	FAssemblyAssociatedAssetDesc NewAssetDesc;
+	NewAssetDesc.AssetClass = AssetClass;
+	NewAssetDesc.AssetName.Template = DefaultName;
+	NewAssetDesc.AssetID = FGuid::NewGuid();
+
+	SchemaToConfigure->TemplateSequence->Modify();
+	SchemaToConfigure->TemplateSequence->AssociatedAssets.Add(NewAssetDesc);
+	TreeView->Reinitialize();
 }
 
 FReply SCineAssemblySchemaWindow::OnAddFolderClicked()
diff --git a/Engine/Plugins/VirtualProduction/CinematicAssemblyTools/Source/CineAssemblyToolsEditor/Private/UI/CineAssembly/SCineAssemblySchemaWindow.h b/Engine/Plugins/VirtualProduction/CinematicAssemblyTools/Source/CineAssemblyToolsEditor/Private/UI/CineAssembly/SCineAssemblySchemaWindow.h
--- a/Engine/Plugins/VirtualProduction/CinematicAssemblyTools/Source/CineAssemblyToolsEditor/Private/UI/CineAssembly/SCineAssemblySchemaWindow.h
+++ b/Engine/Plugins/VirtualProduction/CinematicAssemblyTools/Source/CineAssemblyToolsEditor/Private/UI/CineAssembly/SCineAssemblySchemaWindow.h
@@ -92,8 +92,14 @@ private:
 	/** Closes the window and indicates that no assets should be created by the asset factory */
 	FReply OnCancelClicked();
 
-	/** Adds a new entry to the asset list view */
-	FReply OnAddAssetClicked();
+	/** Returns the menu content for the Add Asset button */
+	TSharedRef<SWidget> OnGetAddAssetMenuContent();
+
+	/** Adds a menu entry to the Add Asset menu for the input asset class */
+	void AddAssetMenuRow(FMenuBuilder& MenuBuilder, const UClass* AssetClass, FExecuteAction OnClickAction = FExecuteAction());
+
+	/** Adds a new associated asset definition to the template sequence */
+	void AddAssociatedAsset(const UClass* AssetClass);
 
 	/** Adds a new folder item to the content tree view */
 	FReply OnAddFolderClicked();
diff --git a/Engine/Plugins/VirtualProduction/CinematicAssemblyTools/Source/CineAssemblyToolsEditor/Public/CineAssemblyEditorFunctionLibrary.h b/Engine/Plugins/VirtualProduction/CinematicAssemblyTools/Source/CineAssemblyToolsEditor/Public/CineAssemblyEditorFunctionLibrary.h
--- a/Engine/Plugins/VirtualProduction/CinematicAssemblyTools/Source/CineAssemblyToolsEditor/Public/CineAssemblyEditorFunctionLibrary.h
+++ b/Engine/Plugins/VirtualProduction/CinematicAssemblyTools/Source/CineAssemblyToolsEditor/Public/CineAssemblyEditorFunctionLibrary.h
@@ -55,17 +55,18 @@ public:
 	/**
 	 * Duplicates a CineAssembly asset at the specified path.
 	 *
-	 * @param SourceAssembly               The assembly to duplicate
-	 * @param DuplicatePath                The content browser path where the duplicate should be created
-	 * @param Metadata                     Optional metadata overrides to apply to the duplicated assembly (key/value string pairs)
-	 * @param bDuplicateSubsequences       If true, recursively duplicate all descendant subsequences. If false, subsequences will reference the same sequences as the source Assembly.
-	 * @param NameOverride                 Optional name for the duplicate. If empty, the source assembly's name is used.
-	 * @param ExternalSequencePreference   Options for handling the duplication of subsequences that live outside of the source assembly's content tree
+	 * @param SourceAssembly                   The assembly to duplicate
+	 * @param DuplicatePath                    The content browser path where the duplicate should be created
+	 * @param Metadata                         Optional metadata overrides to apply to the duplicated assembly (key/value string pairs)
+	 * @param bDuplicateSubsequences           If true, recursively duplicate all descendant subsequences. If false, subsequences will reference the same sequences as the source Assembly.
+	 * @param bDuplicateAssociatedAssets       If true, duplicate all assets associated with this Assembly and its managed SubAssemblies.
+	 * @param NameOverride                     Optional name for the duplicate. If empty, the source assembly's name is used.
+	 * @param ExternalSequencePreference       Options for handling the duplication of subsequences that live outside of the source assembly's content tree
 	 * 
 	 * @return The duplicated assembly, or nullptr if duplication failed
 	 */
 	UFUNCTION(BlueprintCallable, Category = "Cine Assembly Tools", meta = (AutoCreateRefTerm = "Metadata", AdvancedDisplay = "ExternalSequencePreference"))
-	static UCineAssembly* DuplicateAssembly(UCineAssembly* SourceAssembly, const FString& DuplicatePath, const TMap<FString, FString>& Metadata, bool bDuplicateSubsequences = true, const FString& NameOverride = FString(), EDuplicateExternalSequencePreference ExternalSequencePreference = EDuplicateExternalSequencePreference::DuplicateIntoAssemblyFolder);
+	static UCineAssembly* DuplicateAssembly(UCineAssembly* SourceAssembly, const FString& DuplicatePath, const TMap<FString, FString>& Metadata, bool bDuplicateSubsequences = true, bool bDuplicateAssociatedAssets = true, const FString& NameOverride = FString(), EDuplicateExternalSequencePreference ExternalSequencePreference = EDuplicateExternalSequencePreference::DuplicateIntoAssemblyFolder);
 
 	/** Applies each metadata key/value pair in the input map to the input Assembly */
 	UFUNCTION(BlueprintCallable, Category = "Cine Assembly Tools", meta = (AutoCreateRefTerm = "Metadata"))
