ShotManagement: Add associated asset duplication support to the assembly duplication tool.
#rnx
#Synced-CL 51715153

diff --git a/Engine/Plugins/VirtualProduction/CinematicAssemblyTools/Source/CineAssemblyToolsEditor/Private/CineAssemblyFactory.cpp b/Engine/Plugins/VirtualProduction/CinematicAssemblyTools/Source/CineAssemblyToolsEditor/Private/CineAssemblyFactory.cpp
--- a/Engine/Plugins/VirtualProduction/CinematicAssemblyTools/Source/CineAssemblyToolsEditor/Private/CineAssemblyFactory.cpp
+++ b/Engine/Plugins/VirtualProduction/CinematicAssemblyTools/Source/CineAssemblyToolsEditor/Private/CineAssemblyFactory.cpp
@@ -381,46 +381,55 @@ void UCineAssemblyFactory::CreateAssociatedAssets(UCineAssembly* ConfiguredAssem
 	}
 }
 
-void UCineAssemblyFactory::DuplicateAssociatedAssets(UCineAssembly* DuplicatedAssembly)
+void UCineAssemblyFactory::DuplicateAssociatedAsset(UCineAssembly* DuplicatedAssembly, FAssemblyAssociatedAssetDesc& AssetDesc)
 {
 	if (!DuplicatedAssembly)
 	{
 		return;
 	}
 
-	// Only duplicate assets for assemblies that were duplicated from an existing assembly.
-	// Assemblies created fresh from a Schema have their assets created by CreateAssociatedAssets instead.
-	if (!DuplicatedAssembly->SourceAssemblyPath.IsEmpty())
+	UObject* OriginalAsset = AssetDesc.CreatedAsset.LoadSynchronous();
+	if (!OriginalAsset)
 	{
-		FString AssetPath;
-		FString AssemblyRootFolder;
-		DuplicatedAssembly->GetAssetPathAndRootFolder(AssetPath, AssemblyRootFolder);
+		return;
+	}
 
-		IAssetTools& AssetTools = FModuleManager::LoadModuleChecked<FAssetToolsModule>("AssetTools").Get();
+	FString AssetPath;
+	FString RootFolder;
+	DuplicatedAssembly->GetAssetPathAndRootFolder(AssetPath, RootFolder);
 
-		for (FAssemblyAssociatedAssetDesc& AssetDesc : DuplicatedAssembly->AssociatedAssets)
-		{
-			// Only duplicate if the associated asset definition already references a real asset
-			UObject* OriginalAsset = AssetDesc.CreatedAsset.LoadSynchronous();
-			if (!OriginalAsset)
-			{
-				continue;
-			}
+	const FText ResolvedName = UCineAssemblyNamingTokens::GetResolvedText(AssetDesc.AssetName.Template, DuplicatedAssembly);
+	const FText ResolvedPath = UCineAssemblyNamingTokens::GetResolvedText(AssetDesc.RelativePath.Template, DuplicatedAssembly);
 
-			const FText ResolvedName = UCineAssemblyNamingTokens::GetResolvedText(AssetDesc.AssetName.Template, DuplicatedAssembly);
-			const FText ResolvedPath = UCineAssemblyNamingTokens::GetResolvedText(AssetDesc.RelativePath.Template, DuplicatedAssembly);
+	const FString DestinationPath = RootFolder / ResolvedPath.ToString();
 
-			const FString DestinationPath = AssemblyRootFolder / ResolvedPath.ToString();
+	IAssetTools& AssetTools = FModuleManager::LoadModuleChecked<FAssetToolsModule>("AssetTools").Get();
 
-			FString UniquePackageName;
-			FString UniqueAssetName;
-			AssetTools.CreateUniqueAssetName(DestinationPath / ResolvedName.ToString(), FString(), UniquePackageName, UniqueAssetName);
+	FString UniquePackageName;
+	FString UniqueAssetName;
+	AssetTools.CreateUniqueAssetName(DestinationPath / ResolvedName.ToString(), FString(), UniquePackageName, UniqueAssetName);
 
-			UObject* DuplicateAsset = AssetTools.DuplicateAsset(UniqueAssetName, FPaths::GetPath(UniquePackageName), OriginalAsset);
-			if (DuplicateAsset)
-			{
-				AssetDesc.CreatedAsset = DuplicateAsset;
-			}
+	UObject* DuplicateAsset = AssetTools.DuplicateAsset(UniqueAssetName, FPaths::GetPath(UniquePackageName), OriginalAsset);
+	if (DuplicateAsset)
+	{
+		AssetDesc.CreatedAsset = DuplicateAsset;
+	}
+}
+
+void UCineAssemblyFactory::DuplicateAssociatedAssets(UCineAssembly* DuplicatedAssembly)
+{
+	if (!DuplicatedAssembly)
+	{
+		return;
+	}
+
+	// Only duplicate assets for assemblies that were duplicated from an existing assembly.
+	// Assemblies created fresh from a Schema have their assets created by CreateAssociatedAssets instead.
+	if (!DuplicatedAssembly->SourceAssemblyPath.IsEmpty())
+	{
+		for (FAssemblyAssociatedAssetDesc& AssetDesc : DuplicatedAssembly->AssociatedAssets)
+		{
+			DuplicateAssociatedAsset(DuplicatedAssembly, AssetDesc);
 		}
 	}
 
diff --git a/Engine/Plugins/VirtualProduction/CinematicAssemblyTools/Source/CineAssemblyToolsEditor/Private/CineAssemblyFactory.h b/Engine/Plugins/VirtualProduction/CinematicAssemblyTools/Source/CineAssemblyToolsEditor/Private/CineAssemblyFactory.h
--- a/Engine/Plugins/VirtualProduction/CinematicAssemblyTools/Source/CineAssemblyToolsEditor/Private/CineAssemblyFactory.h
+++ b/Engine/Plugins/VirtualProduction/CinematicAssemblyTools/Source/CineAssemblyToolsEditor/Private/CineAssemblyFactory.h
@@ -43,6 +43,9 @@ public:
 	/** Create new assets for each associated asset definition on the input Assembly */
 	static void CreateAssociatedAssets(UCineAssembly* ConfiguredAssembly);
 
+	/** Duplicate the input asset associated with the duplicated assembly */
+	static void DuplicateAssociatedAsset(UCineAssembly* DuplicatedAssembly, FAssemblyAssociatedAssetDesc& AssetDesc);
+
 	/** Duplicate the associated assets of a duplicated assembly, updating each definition's CreatedAsset to point to the duplicate */
 	static void DuplicateAssociatedAssets(UCineAssembly* DuplicatedAssembly);
 
diff --git a/Engine/Plugins/VirtualProduction/CinematicAssemblyTools/Source/CineAssemblyToolsEditor/Private/UI/CineAssembly/SDuplicateAssemblyWindow.cpp b/Engine/Plugins/VirtualProduction/CinematicAssemblyTools/Source/CineAssemblyToolsEditor/Private/UI/CineAssembly/SDuplicateAssemblyWindow.cpp
--- a/Engine/Plugins/VirtualProduction/CinematicAssemblyTools/Source/CineAssemblyToolsEditor/Private/UI/CineAssembly/SDuplicateAssemblyWindow.cpp
+++ b/Engine/Plugins/VirtualProduction/CinematicAssemblyTools/Source/CineAssemblyToolsEditor/Private/UI/CineAssembly/SDuplicateAssemblyWindow.cpp
@@ -2,6 +2,7 @@
 
 #include "SDuplicateAssemblyWindow.h"
 
+#include "Algo/Find.h"
 #include "AssetToolsModule.h"
 #include "CineAssemblyFactory.h"
 #include "CineAssemblyNamingTokens.h"
@@ -16,6 +17,7 @@
 #include "Sections/MovieSceneSubSection.h"
 #include "Sequencer/CineAssemblySchemaSequencerUtilities.h"
 #include "SNamingTokensEditableTextBox.h"
+#include "Styling/SlateIconFinder.h"
 #include "Tracks/MovieSceneSubTrack.h"
 #include "Widgets/Input/SEditableTextBox.h"
 #include "Widgets/Layout/SSeparator.h"
@@ -296,6 +298,12 @@ FReply SDuplicateAssemblyWindow::OnDuplicateClicked()
 	// Duplicate the remaining subsequences (those that are not Assemblies) that are set to DuplicateOriginal
 	DuplicateSubsequencesRecursive(DuplicateAssembly.Get());
 
+	// Duplicate, maintain, or clear the associated assets
+	for (TSharedPtr<FSequenceTreeItem>& Item : TreeItems)
+	{
+		ProcessAssociatedAssets(Item);
+	}
+
 	// Remove any subsequences that are set to Remove
 	RemoveSubsequencesRecursive(DuplicateAssembly.Get());
 
@@ -327,6 +335,8 @@ void SDuplicateAssemblyWindow::PopulateSubsequenceTreeRecursive(UMovieSceneSeque
 		return;
 	}
 
+	UCineAssembly* OwningAssembly = Cast<UCineAssembly>(InSequence);
+
 	// Loop through all of the sections to find every subsequence in the input sequence
 	for (UMovieSceneSection* Section : MovieScene->GetAllSections())
 	{
@@ -339,12 +349,11 @@ void SDuplicateAssemblyWindow::PopulateSubsequenceTreeRecursive(UMovieSceneSeque
 				NewItem->OriginalSequence = Subsequence;
 				NewItem->Parent = InItem;
 
-				NewItem->DuplicationData.SubsequenceName = Subsequence->IsA<UCineAssembly>() ? Cast<UCineAssembly>(Subsequence)->AssemblyName.Template : Subsequence->GetDisplayName().ToString();
+				NewItem->DuplicationData.DuplicateName = Subsequence->IsA<UCineAssembly>() ? Cast<UCineAssembly>(Subsequence)->AssemblyName.Template : Subsequence->GetDisplayName().ToString();
 
 				InItem ? InItem->Children.Add(NewItem) : TreeItems.Add(NewItem);
 
-				// We handle managed SubAssemblies differently than other Subsequences found in the input Sequence's MovieScene 
-				UCineAssembly* OwningAssembly = Cast<UCineAssembly>(InSequence);
+				// We handle managed SubAssemblies differently than other Subsequences found in the input Sequence's MovieScene
 				if (OwningAssembly && OwningAssembly->SubAssemblies.Contains(SubSection))
 				{
 					// Pre-duplicate this SubAssembly such that it can be configured by the user before assets are created
@@ -356,6 +365,7 @@ void SDuplicateAssemblyWindow::PopulateSubsequenceTreeRecursive(UMovieSceneSeque
 
 					FCineAssemblySchemaSequencerUtilities::ReplaceSubSequence(SubSection, DuplicateSubAssembly);
 
+					NewItem->OwningAssembly = OwningAssembly;
 					NewItem->bIsManagedSubAssembly = true;
 					NewItem->DuplicationData.DuplicationMode = GetDuplicationPreference(DuplicateSubAssembly->AssemblyName.Template);
 				}
@@ -386,6 +396,26 @@ void SDuplicateAssemblyWindow::PopulateSubsequenceTreeRecursive(UMovieSceneSeque
 			}
 		}
 	}
+
+	// Add associated asset items for any assembly that has created assets
+	if (OwningAssembly)
+	{
+		for (const FAssemblyAssociatedAssetDesc& AssetDesc : OwningAssembly->AssociatedAssets)
+		{
+			if (AssetDesc.CreatedAsset.IsNull())
+			{
+				continue;
+			}
+
+			TSharedPtr<FSequenceTreeItem> NewItem = MakeShared<FSequenceTreeItem>();
+			NewItem->AssociatedAssetID = AssetDesc.AssetID;
+			NewItem->OwningAssembly = OwningAssembly;
+			NewItem->DuplicationData.DuplicateName = AssetDesc.AssetName.Template;
+			NewItem->Parent = InItem;
+
+			InItem ? InItem->Children.Add(NewItem) : TreeItems.Add(NewItem);
+		}
+	}
 }
 
 void SDuplicateAssemblyWindow::DuplicateSubsequencesRecursive(UMovieSceneSequence* InSequence)
@@ -416,7 +446,7 @@ void SDuplicateAssemblyWindow::DuplicateSubsequencesRecursive(UMovieSceneSequenc
 				{
 					if (DuplicationData->DuplicationMode == ECineAssemblyDuplicationMode::DuplicateOriginal)
 					{
-						const FString ResolvedName = UCineAssemblyNamingTokens::GetResolvedText(DuplicationData->SubsequenceName, DuplicateAssembly.Get()).ToString();
+						const FString ResolvedName = UCineAssemblyNamingTokens::GetResolvedText(DuplicationData->DuplicateName, DuplicateAssembly.Get()).ToString();
 
 						// Use the current parent sequence's path as the base, so that subsequences inside a renamed SubAssembly are placed relative to the SubAssembly's new folder
 						const FString ParentAssetPath = FPaths::GetPath(InSequence->GetPathName());
@@ -543,8 +573,9 @@ void SDuplicateAssemblyWindow::RestoreOriginalSubAssembliesRecursive(const TShar
 void SDuplicateAssemblyWindow::BuildDuplicationMapRecursive(const TSharedPtr<FSequenceTreeItem> InItem)
 {
 	// SubAssemblies have already been duplicated so they should not be added to the map to avoid being duplicated again
+	// Associated assets are also handled separately and don't belong in the subsequence duplication map
 	const bool bIsAlreadyDuplicated = InItem->bIsManagedSubAssembly && (InItem->DuplicationData.DuplicationMode == ECineAssemblyDuplicationMode::DuplicateOriginal);
-	if (!bIsAlreadyDuplicated)
+	if (!bIsAlreadyDuplicated && !InItem->AssociatedAssetID.IsValid())
 	{
 		SubsequenceDuplicationData.Add(InItem->OriginalSequence, InItem->DuplicationData);
 	}
@@ -619,16 +650,71 @@ ECineAssemblyDuplicationMode SDuplicateAssemblyWindow::GetDuplicationPreference(
 	return ECineAssemblyDuplicationMode::DuplicateOriginal;
 }
 
+void SDuplicateAssemblyWindow::ProcessAssociatedAssets(const TSharedPtr<FSequenceTreeItem>& InItem)
+{
+	if (InItem->AssociatedAssetID.IsValid())
+	{
+		UCineAssembly* Assembly = InItem->OwningAssembly.Get();
+		if (!Assembly)
+		{
+			return;
+		}
+
+		FAssemblyAssociatedAssetDesc* AssetDesc = Algo::FindBy(Assembly->AssociatedAssets, InItem->AssociatedAssetID, &FAssemblyAssociatedAssetDesc::AssetID);
+		if (!AssetDesc)
+		{
+			return;
+		}
+
+		switch (InItem->DuplicationData.DuplicationMode)
+		{
+		case ECineAssemblyDuplicationMode::DuplicateOriginal:
+			// Update the AssetDesc name to match the value of the Item
+			AssetDesc->AssetName.Template = InItem->DuplicationData.DuplicateName;
+			UCineAssemblyFactory::DuplicateAssociatedAsset(Assembly, *AssetDesc);
+			break;
+
+		case ECineAssemblyDuplicationMode::MaintainReference:
+			// No-op: keep the existing CreatedAsset reference
+			break;
+
+		case ECineAssemblyDuplicationMode::Remove:
+			AssetDesc->CreatedAsset = nullptr;
+			break;
+
+		default:
+			checkNoEntry();
+		}
+	}
+
+	for (const TSharedPtr<FSequenceTreeItem>& Child : InItem->Children)
+	{
+		ProcessAssociatedAssets(Child);
+	}
+}
+
 void SSubsequenceDuplicationRow::Construct(const FArguments& Args, const TSharedRef<STableViewBase>& OwnerTableView, const TSharedPtr<FSequenceTreeItem>& InTreeItem, UCineAssembly* DuplicateAssembly)
 {
 	TreeItem = InTreeItem;
 
+	if (TreeItem->AssociatedAssetID.IsValid() && TreeItem->OwningAssembly.IsValid())
+	{
+		if (const FAssemblyAssociatedAssetDesc* AssetDesc = Algo::FindBy(TreeItem->OwningAssembly->AssociatedAssets, TreeItem->AssociatedAssetID, &FAssemblyAssociatedAssetDesc::AssetID))
+		{
+			ItemClass = AssetDesc->AssetClass.LoadSynchronous();
+		}
+	}
+
 	NamingTokenContext = TStrongObjectPtr<UCineAssemblyNamingTokensContext>(NewObject<UCineAssemblyNamingTokensContext>());
 
 	if (UCineAssembly* DuplicateSubAssembly = TreeItem->GetDuplicateSubAssembly())
 	{
 		NamingTokenContext->Assembly = DuplicateSubAssembly;
 	}
+	else if (TreeItem->AssociatedAssetID.IsValid())
+	{
+		NamingTokenContext->Assembly = TreeItem->OwningAssembly.Get();
+	}
 	else
 	{
 		NamingTokenContext->Assembly = DuplicateAssembly;
@@ -677,14 +763,14 @@ TSharedRef<SWidget> SSubsequenceDuplicationRow::MakeOriginalNameWidget()
 		.VAlign(VAlign_Center)
 		.Padding(8.0f, 0.0f, 0.0f, 0.0f)
 		[
-			SNew(SImage).Image(FCineAssemblyToolsStyle::Get().GetBrush("Icons.Sequencer"))
+			SNew(SImage).Image(this, &SSubsequenceDuplicationRow::GetOriginalNameIcon)
 		]
 
 		+ SHorizontalBox::Slot()
 		.VAlign(VAlign_Center)
 		.Padding(8.0f, 0.0f, 0.0f, 0.0f)
 		[
-			SNew(STextBlock).Text(TreeItem->OriginalSequence->GetDisplayName())
+			SNew(STextBlock).Text(this, &SSubsequenceDuplicationRow::GetOriginalDisplayName)
 		];
 }
 
@@ -756,14 +842,7 @@ TSharedRef<SWidget> SSubsequenceDuplicationRow::MakeDuplicateNameWidget()
 					{
 						return (TreeItem->DuplicationData.DuplicationMode != ECineAssemblyDuplicationMode::Remove) ? EVisibility::Visible : EVisibility::Hidden;
 					})
-				.Text_Lambda([this]()
-					{
-						if (TreeItem->DuplicationData.DuplicationMode == ECineAssemblyDuplicationMode::DuplicateOriginal)
-						{
-							return FText::FromString(TreeItem->DuplicationData.SubsequenceName);
-						}
-						return TreeItem->OriginalSequence->GetDisplayName();
-					})
+				.Text(this, &SSubsequenceDuplicationRow::GetDuplicateDisplayName)
 				.OnTextCommitted(this, &SSubsequenceDuplicationRow::OnSubsequenceNameCommitted)
 				.OnTokenizedTextEvaluated(this, &SSubsequenceDuplicationRow::OnSubsequenceTokenTextEvaluated)
 				.OnValidateTokenizedText(this, &SSubsequenceDuplicationRow::ValidateSubsequenceName)
@@ -799,9 +878,61 @@ void SSubsequenceDuplicationRow::SetDuplicationModeRecursive(TSharedPtr<FSequenc
 	}
 }
 
+const FSlateBrush* SSubsequenceDuplicationRow::GetOriginalNameIcon() const
+{
+	if (ItemClass)
+	{
+		return FSlateIconFinder::FindIconBrushForClass(ItemClass);
+	}
+	return FCineAssemblyToolsStyle::Get().GetBrush("Icons.Sequencer");
+}
+
+FText SSubsequenceDuplicationRow::GetOriginalDisplayName() const
+{
+	if (TreeItem->AssociatedAssetID.IsValid() && TreeItem->OwningAssembly.IsValid())
+	{
+		if (const FAssemblyAssociatedAssetDesc* AssetDesc = Algo::FindBy(TreeItem->OwningAssembly->AssociatedAssets, TreeItem->AssociatedAssetID, &FAssemblyAssociatedAssetDesc::AssetID))
+		{
+			return FText::FromString(AssetDesc->CreatedAsset.GetAssetName());
+		}
+		return FText::GetEmpty();
+	}
+
+	if (TreeItem->OriginalSequence)
+	{
+		return TreeItem->OriginalSequence->GetDisplayName();
+	}
+
+	return FText::GetEmpty();
+}
+
+FText SSubsequenceDuplicationRow::GetDuplicateDisplayName() const
+{
+	if (TreeItem->DuplicationData.DuplicationMode == ECineAssemblyDuplicationMode::DuplicateOriginal)
+	{
+		return FText::FromString(TreeItem->DuplicationData.DuplicateName);
+	}
+
+	if (TreeItem->AssociatedAssetID.IsValid() && TreeItem->OwningAssembly.IsValid())
+	{
+		if (const FAssemblyAssociatedAssetDesc* AssetDesc = Algo::FindBy(TreeItem->OwningAssembly->AssociatedAssets, TreeItem->AssociatedAssetID, &FAssemblyAssociatedAssetDesc::AssetID))
+		{
+			return FText::FromString(AssetDesc->AssetName.Template);
+		}
+		return FText::GetEmpty();
+	}
+
+	if (TreeItem->OriginalSequence)
+	{
+		return TreeItem->OriginalSequence->GetDisplayName();
+	}
+
+	return FText::GetEmpty();
+}
+
 void SSubsequenceDuplicationRow::OnSubsequenceNameCommitted(const FText& InText, ETextCommit::Type InCommitType)
 {
-	TreeItem->DuplicationData.SubsequenceName = InText.ToString();
+	TreeItem->DuplicationData.DuplicateName = InText.ToString();
 
 	if (UCineAssembly* DuplicateSubAssembly = TreeItem->GetDuplicateSubAssembly())
 	{
diff --git a/Engine/Plugins/VirtualProduction/CinematicAssemblyTools/Source/CineAssemblyToolsEditor/Private/UI/CineAssembly/SDuplicateAssemblyWindow.h b/Engine/Plugins/VirtualProduction/CinematicAssemblyTools/Source/CineAssemblyToolsEditor/Private/UI/CineAssembly/SDuplicateAssemblyWindow.h
--- a/Engine/Plugins/VirtualProduction/CinematicAssemblyTools/Source/CineAssemblyToolsEditor/Private/UI/CineAssembly/SDuplicateAssemblyWindow.h
+++ b/Engine/Plugins/VirtualProduction/CinematicAssemblyTools/Source/CineAssemblyToolsEditor/Private/UI/CineAssembly/SDuplicateAssemblyWindow.h
@@ -31,12 +31,19 @@ struct FSubsequenceDuplicationData
 	FString RelativePath;
 
 	/** The name to use for the duplicated subsequence (if the mode is DuplicateOriginal) */
-	FString SubsequenceName;
+	FString DuplicateName;
 };
 
 /** An item in the subsequence tree view, associating a subsequence with its duplication data and child subsequences */
 struct FSequenceTreeItem
 {
+	/**
+	 * The assembly that owns this tree item. This could be null if this item is not managed by a CineAssembly
+	 * For SubAssembly items, this is the assembly that manages this item in its SubAssemblies list.
+	 * For Associated Asset items, this is the assembly that manages this item in its AssociatedAssets list.
+	 */
+	TWeakObjectPtr<UCineAssembly> OwningAssembly;
+
 	/** The SubSection in the MovieScene of the Duplicate Assembly (or one of its SubAssemblies) */
 	UMovieSceneSubSection* SubSection = nullptr;
 
@@ -46,6 +53,9 @@ struct FSequenceTreeItem
 	/** The duplication data associated with this subsequence */
 	FSubsequenceDuplicationData DuplicationData;
 
+	/** The ID of the associated asset in the owning assembly's AssociatedAssets array. Invalid for SubSequence item types. */
+	FGuid AssociatedAssetID;
+
 	/** Child subsequence tree items */
 	TArray<TSharedPtr<FSequenceTreeItem>> Children;
 
@@ -78,6 +88,15 @@ private:
 	TSharedRef<SWidget> MakeDuplicationModeWidget();
 	TSharedRef<SWidget> MakeDuplicateNameWidget();
 
+	/** Returns the icon brush for the original name column */
+	const FSlateBrush* GetOriginalNameIcon() const;
+
+	/** Returns the display name for the original name column */
+	FText GetOriginalDisplayName() const;
+
+	/** Returns the display name for the duplicate name column (varies by duplication mode) */
+	FText GetDuplicateDisplayName() const;
+
 	/** Updates the duplication mode for the input tree item, and potentially its children */
 	void SetDuplicationModeRecursive(TSharedPtr<FSequenceTreeItem> InItem, ECineAssemblyDuplicationMode Mode);
 
@@ -102,6 +121,9 @@ private:
 
 	/** Naming Token filter properties */
 	FNamingTokenFilterArgs FilterArgs;
+
+	/** Cached item class */
+	UClass* ItemClass = nullptr;
 };
 
 /** A window to configure how a Cine Assembly asset should be duplicated */
@@ -157,6 +179,9 @@ private:
 	/** Recursively removes the set of tracks from the input folders (or child folders) */
 	void RemoveTracksFromFoldersRecursive(TArrayView<UMovieSceneFolder* const> Folders, const TArray<UMovieSceneTrack*>& TracksToRemove);
 
+	/** Recursively handles duplication, maintaining, or clearing of the associated assets based on the duplication mode */
+	void ProcessAssociatedAssets(const TSharedPtr<FSequenceTreeItem>& InItem);
+
 	/** Converts the subsequence tree into a map of subsequences to their associated duplication data */
 	void BuildDuplicationMapRecursive(const TSharedPtr<FSequenceTreeItem> InTreeItem);
 
