ShotManagement: Add DefaultLevel metadata to the schema definition, allowing users to specify the default level for new assemblies made from that schema.
#rnx
#Synced-CL 51715153

diff --git a/Engine/Plugins/VirtualProduction/CinematicAssemblyTools/Source/CineAssemblyTools/Private/CineAssembly.cpp b/Engine/Plugins/VirtualProduction/CinematicAssemblyTools/Source/CineAssemblyTools/Private/CineAssembly.cpp
--- a/Engine/Plugins/VirtualProduction/CinematicAssemblyTools/Source/CineAssemblyTools/Private/CineAssembly.cpp
+++ b/Engine/Plugins/VirtualProduction/CinematicAssemblyTools/Source/CineAssemblyTools/Private/CineAssembly.cpp
@@ -319,10 +319,16 @@ void UCineAssembly::ChangeSchema(UCineAssemblySchema* InSchema)
 
 	SubAssemblies.Empty();
 
-	// Update the Data Only flag on the assembly
 	if (BaseSchema)
 	{
+		// Update the Data Only flag on the assembly
 		bIsDataOnly = BaseSchema->bIsDataOnly;
+
+		// Apply the schema's default Level if one is specified
+		if (BaseSchema->DefaultLevel.IsSet())
+		{
+			Level = BaseSchema->DefaultLevel.GetValue();
+		}
 	}
 }
 
diff --git a/Engine/Plugins/VirtualProduction/CinematicAssemblyTools/Source/CineAssemblyTools/Private/CineAssemblySchema.cpp b/Engine/Plugins/VirtualProduction/CinematicAssemblyTools/Source/CineAssemblyTools/Private/CineAssemblySchema.cpp
--- a/Engine/Plugins/VirtualProduction/CinematicAssemblyTools/Source/CineAssemblyTools/Private/CineAssemblySchema.cpp
+++ b/Engine/Plugins/VirtualProduction/CinematicAssemblyTools/Source/CineAssemblyTools/Private/CineAssemblySchema.cpp
@@ -173,6 +173,26 @@ void UCineAssemblySchema::SetDefaultAssemblyPath(const FString& InPath)
 	}
 }
 
+TSoftObjectPtr<UWorld> UCineAssemblySchema::GetDefaultLevel() const
+{
+	const FSoftObjectPath* DefaultLevelPtr = DefaultLevel.GetPtrOrNull();
+	if (DefaultLevelPtr)
+	{
+		return TSoftObjectPtr<UWorld>(*DefaultLevelPtr);
+	}
+	return nullptr;
+}
+
+void UCineAssemblySchema::SetDefaultLevel(TSoftObjectPtr<UWorld> InLevel)
+{
+	DefaultLevel = InLevel.ToSoftObjectPath();
+}
+
+void UCineAssemblySchema::ClearDefaultLevel()
+{
+	DefaultLevel.Reset();
+}
+
 TArray<FString> UCineAssemblySchema::GetDefaultFolders() const
 {
 	TArray<FString> FolderPaths;
diff --git a/Engine/Plugins/VirtualProduction/CinematicAssemblyTools/Source/CineAssemblyTools/Public/CineAssemblySchema.h b/Engine/Plugins/VirtualProduction/CinematicAssemblyTools/Source/CineAssemblyTools/Public/CineAssemblySchema.h
--- a/Engine/Plugins/VirtualProduction/CinematicAssemblyTools/Source/CineAssemblyTools/Public/CineAssemblySchema.h
+++ b/Engine/Plugins/VirtualProduction/CinematicAssemblyTools/Source/CineAssemblyTools/Public/CineAssemblySchema.h
@@ -84,6 +84,18 @@ public:
 	UFUNCTION(BlueprintCallable, Category = "Cine Assembly Tools")
 	CINEASSEMBLYTOOLS_API void SetDefaultAssemblyPath(const FString& InPath);
 
+	/** Get the default level for assemblies created from this schema */
+	UFUNCTION(BlueprintPure, Category = "Cine Assembly Tools")
+	CINEASSEMBLYTOOLS_API TSoftObjectPtr<UWorld> GetDefaultLevel() const;
+
+	/** Set the default level for assemblies created from this schema */
+	UFUNCTION(BlueprintCallable, Category = "Cine Assembly Tools")
+	CINEASSEMBLYTOOLS_API void SetDefaultLevel(TSoftObjectPtr<UWorld> InLevel);
+
+	/** Clear the default level for assemblies created from this schema */
+	UFUNCTION(BlueprintCallable, Category = "Cine Assembly Tools")
+	CINEASSEMBLYTOOLS_API void ClearDefaultLevel();
+
 	/** Get the list of default folder paths that will be created for assemblies made from this schema */
 	UFUNCTION(BlueprintPure, Category = "Cine Assembly Tools")
 	CINEASSEMBLYTOOLS_API TArray<FString> GetDefaultFolders() const;
@@ -128,6 +140,10 @@ public:
 	UPROPERTY(EditAnywhere, BlueprintReadWrite, Category = "Default")
 	TObjectPtr<UTexture2D> ThumbnailImage;
 
+	/** Optional default Level to assign to assemblies created from this schema */
+	UPROPERTY(EditAnywhere, Category = "Assembly Defaults", meta = (AllowedClasses = "/Script/Engine.World"))
+	TOptional<FSoftObjectPath> DefaultLevel;
+
 	/** List of metadata fields that should be automatically added to assemblies made from this schema */
 	UPROPERTY(EditAnywhere, Category = "Metadata")
 	TArray<FAssemblyMetadataDesc> AssemblyMetadata;
diff --git a/Engine/Plugins/VirtualProduction/CinematicAssemblyTools/Source/CineAssemblyToolsEditor/Private/UI/CineAssembly/SCineAssemblySchemaWindow.cpp b/Engine/Plugins/VirtualProduction/CinematicAssemblyTools/Source/CineAssemblyToolsEditor/Private/UI/CineAssembly/SCineAssemblySchemaWindow.cpp
--- a/Engine/Plugins/VirtualProduction/CinematicAssemblyTools/Source/CineAssemblyToolsEditor/Private/UI/CineAssembly/SCineAssemblySchemaWindow.cpp
+++ b/Engine/Plugins/VirtualProduction/CinematicAssemblyTools/Source/CineAssemblyToolsEditor/Private/UI/CineAssembly/SCineAssemblySchemaWindow.cpp
@@ -480,6 +480,10 @@ TSharedRef<SWidget> SCineAssemblySchemaWindow::MakeDetailsWidget(bool bShowMetad
 
 bool SCineAssemblySchemaWindow::IsPropertyVisible(const FPropertyAndParent& PropertyAndParent, bool bShowMetadata)
 {
+	if (PropertyAndParent.Property.GetFName() == GET_MEMBER_NAME_CHECKED(UCineAssemblySchema, DefaultLevel))
+	{
+		return bShowMetadata;
+	}
 	if (PropertyAndParent.Property.GetFName() == GET_MEMBER_NAME_CHECKED(UCineAssemblySchema, AssemblyMetadata))
 	{
 		return bShowMetadata;
