ShotManagement: Fix ensure when using a schema that has a metadata field of type AssetPath whose class is null.
#rnx
#Synced-CL 51715153

diff --git a/Engine/Plugins/VirtualProduction/CinematicAssemblyTools/Source/CineAssemblyToolsEditor/Private/CineAssemblyCustomization.cpp b/Engine/Plugins/VirtualProduction/CinematicAssemblyTools/Source/CineAssemblyToolsEditor/Private/CineAssemblyCustomization.cpp
--- a/Engine/Plugins/VirtualProduction/CinematicAssemblyTools/Source/CineAssemblyToolsEditor/Private/CineAssemblyCustomization.cpp
+++ b/Engine/Plugins/VirtualProduction/CinematicAssemblyTools/Source/CineAssemblyToolsEditor/Private/CineAssemblyCustomization.cpp
@@ -326,7 +326,7 @@ TSharedRef<SWidget> FCineAssemblyCustomization::MakeMetadataWidget(IDetailLayout
 					});
 		case ECineAssemblyMetadataType::AssetPath:
 			return SNew(SObjectPropertyEntryBox)
-				.AllowedClass(MetadataDesc.AssetClass.ResolveClass())
+				.AllowedClass(!MetadataDesc.AssetClass.IsNull() ? MetadataDesc.AssetClass.ResolveClass() : UObject::StaticClass())
 				.ThumbnailPool(DetailBuilder.GetThumbnailPool())
 				.AllowCreate(true)
 				.ObjectPath_Lambda([this, MetadataDesc]()
