ShotManagement: Add support for linking Schema metadata fields to Schema defined asset templates. After a linked asset is created alongside an Assembly, that Assembly's linked metadata will be set to point to the linked asset.
#rnx

#Synced-CL 51715153

diff --git a/Engine/Plugins/VirtualProduction/CinematicAssemblyTools/Source/CineAssemblyTools/Private/CineAssembly.cpp b/Engine/Plugins/VirtualProduction/CinematicAssemblyTools/Source/CineAssemblyTools/Private/CineAssembly.cpp
--- a/Engine/Plugins/VirtualProduction/CinematicAssemblyTools/Source/CineAssemblyTools/Private/CineAssembly.cpp
+++ b/Engine/Plugins/VirtualProduction/CinematicAssemblyTools/Source/CineAssemblyTools/Private/CineAssembly.cpp
@@ -193,6 +193,7 @@ void UCineAssembly::InitializeFromTemplate(ULevelSequence* Template)
 		PathRelativeToRoot = TemplateAssembly->PathRelativeToRoot;
 		DefaultFolderNames = TemplateAssembly->DefaultFolderNames;
 		AssociatedAssets = TemplateAssembly->AssociatedAssets;
+		MetadataLinks = TemplateAssembly->MetadataLinks;
 	}
 
 	if (UMovieScene* TemplateMovieScene = Template->GetMovieScene())
@@ -381,6 +382,23 @@ UCineAssembly* UCineAssembly::CreateSubAssemblyFromTemplate(UObject* TemplateObj
 	return SubAssembly;
 }
 
+void UCineAssembly::RemoveAssociatedAsset(const FGuid& InAssetID)
+{
+	Modify();
+
+	// Remove the input asset from the associate asset list
+	AssociatedAssets.RemoveAll([&InAssetID](const FAssemblyAssociatedAssetDesc& AssetDesc) { return AssetDesc.AssetID == InAssetID; });
+
+	// Remove any metadata links that link to the removed asset
+	for (auto It = MetadataLinks.CreateIterator(); It; ++It)
+	{
+		if (It->Value == InAssetID)
+		{
+			It.RemoveCurrent();
+		}
+	}
+}
+
 void UCineAssembly::DuplicateManagedSubAssemblies()
 {
 	if (!MovieScene)
diff --git a/Engine/Plugins/VirtualProduction/CinematicAssemblyTools/Source/CineAssemblyTools/Public/CineAssembly.h b/Engine/Plugins/VirtualProduction/CinematicAssemblyTools/Source/CineAssemblyTools/Public/CineAssembly.h
--- a/Engine/Plugins/VirtualProduction/CinematicAssemblyTools/Source/CineAssemblyTools/Public/CineAssembly.h
+++ b/Engine/Plugins/VirtualProduction/CinematicAssemblyTools/Source/CineAssemblyTools/Public/CineAssembly.h
@@ -220,6 +220,9 @@ public:
 	/** Returns the asset path and its root folder (which might be different if the schema defines a default assembly path) */
 	CINEASSEMBLYTOOLS_API void GetAssetPathAndRootFolder(FString& OutAssetPath, FString& OutRootFolder);
 
+	/** Removes an associated asset by its AssetID, and cleans up any MetadataLinks that referenced it */
+	CINEASSEMBLYTOOLS_API void RemoveAssociatedAsset(const FGuid& InAssetID);
+
 	/** Recursively duplicates all managed SubAssemblies, replacing each with a new transient copy */
 	CINEASSEMBLYTOOLS_API void DuplicateManagedSubAssemblies();
 
@@ -299,6 +302,10 @@ public:
 	UPROPERTY()
 	TArray<FAssemblyAssociatedAssetDesc> AssociatedAssets;
 
+	/** Links metadata keys to the associated assets to auto-populate the value of the linked metadata field when the asset is created */
+	UPROPERTY()
+	TMap<FString, FGuid> MetadataLinks;
+
 	/** Array of Subsequence Sections created based on the schema */
 	UPROPERTY()
 	TArray<TObjectPtr<UMovieSceneSubSection>> SubAssemblies;
diff --git a/Engine/Plugins/VirtualProduction/CinematicAssemblyTools/Source/CineAssemblyToolsEditor/Private/CineAssemblyFactory.cpp b/Engine/Plugins/VirtualProduction/CinematicAssemblyTools/Source/CineAssemblyToolsEditor/Private/CineAssemblyFactory.cpp
--- a/Engine/Plugins/VirtualProduction/CinematicAssemblyTools/Source/CineAssemblyToolsEditor/Private/CineAssemblyFactory.cpp
+++ b/Engine/Plugins/VirtualProduction/CinematicAssemblyTools/Source/CineAssemblyToolsEditor/Private/CineAssemblyFactory.cpp
@@ -2,6 +2,7 @@
 
 #include "CineAssemblyFactory.h"
 
+#include "Algo/Find.h"
 #include "AssetRegistry/AssetRegistryModule.h"
 #include "AssetViewUtils.h"
 #include "AssetToolsModule.h"
@@ -38,6 +39,27 @@ namespace UE::CineAssemblyTools::Private
 		return nullptr;
 	}
 
+	/** Updates metadata fields on the assembly that are linked to associated assets */
+	void PopulateLinkedMetadata(UCineAssembly* Assembly)
+	{
+		for (const TPair<FString, FGuid>& MetadataLink : Assembly->MetadataLinks)
+		{
+			if (!MetadataLink.Value.IsValid())
+			{
+				continue;
+			}
+
+			// Find the associated asset that matches the linked GUID
+			if (const FAssemblyAssociatedAssetDesc* AssetDesc = Algo::FindBy(Assembly->AssociatedAssets, MetadataLink.Value, &FAssemblyAssociatedAssetDesc::AssetID))
+			{
+				if (!AssetDesc->CreatedAsset.IsNull())
+				{
+					Assembly->SetMetadataAsString(MetadataLink.Key, AssetDesc->CreatedAsset.ToSoftObjectPath().ToString());
+				}
+			}
+		}
+	}
+
 	/** Returns the desired asset name and path for a duplicate of SequenceToDuplicate, placed relative to ParentSequence */
 	FString GetDuplicateNameAndPath(UMovieSceneSequence* SequenceToDuplicate, UMovieSceneSequence* ParentSequence, const TMap<UMovieSceneSequence*, FString>& OriginalPaths)
 	{
@@ -379,6 +401,9 @@ void UCineAssemblyFactory::CreateAssociatedAssets(UCineAssembly* ConfiguredAssem
 			AssetDesc.CreatedAsset = NewAsset;
 		}
 	}
+
+	// Auto-populate metadata fields that are linked to associated assets
+	UE::CineAssemblyTools::Private::PopulateLinkedMetadata(ConfiguredAssembly);
 }
 
 void UCineAssemblyFactory::DuplicateAssociatedAsset(UCineAssembly* DuplicatedAssembly, FAssemblyAssociatedAssetDesc& AssetDesc)
@@ -433,10 +458,14 @@ void UCineAssemblyFactory::DuplicateAssociatedAssets(UCineAssembly* DuplicatedAs
 		}
 	}
 
+	// Auto-populate metadata fields that are linked to associated assets
+	UE::CineAssemblyTools::Private::PopulateLinkedMetadata(DuplicatedAssembly);
+
 	// Recurse into managed SubAssemblies
 	for (UMovieSceneSubSection* SubSection : DuplicatedAssembly->SubAssemblies)
 	{
-		if (UCineAssembly* SubAssembly = SubSection ? Cast<UCineAssembly>(SubSection->GetSequence()) : nullptr)
+		UCineAssembly* SubAssembly = SubSection ? Cast<UCineAssembly>(SubSection->GetSequence()) : nullptr;
+		if (SubAssembly)
 		{
 			DuplicateAssociatedAssets(SubAssembly);
 		}
diff --git a/Engine/Plugins/VirtualProduction/CinematicAssemblyTools/Source/CineAssemblyToolsEditor/Private/CineAssemblyMetadataCustomization.cpp b/Engine/Plugins/VirtualProduction/CinematicAssemblyTools/Source/CineAssemblyToolsEditor/Private/CineAssemblyMetadataCustomization.cpp
--- a/Engine/Plugins/VirtualProduction/CinematicAssemblyTools/Source/CineAssemblyToolsEditor/Private/CineAssemblyMetadataCustomization.cpp
+++ b/Engine/Plugins/VirtualProduction/CinematicAssemblyTools/Source/CineAssemblyToolsEditor/Private/CineAssemblyMetadataCustomization.cpp
@@ -3,6 +3,7 @@
 #include "CineAssemblyMetadataCustomization.h"
 
 #include "Algo/Contains.h"
+#include "Algo/Find.h"
 #include "Algo/IndexOf.h"
 #include "CineAssembly.h"
 #include "CineAssemblyNamingTokens.h"
@@ -12,6 +13,7 @@
 #include "NamingTokensEngineSubsystem.h"
 #include "PropertyCustomizationHelpers.h"
 #include "SNamingTokensEditableTextBox.h"
+#include "UI/SAssemblyMetadataLink.h"
 #include "Widgets/Input/SCheckBox.h"
 #include "Widgets/Input/SEditableTextBox.h"
 #include "Widgets/Input/SMultiLineEditableTextBox.h"
@@ -143,11 +145,16 @@ void FCineAssemblyMetadataCustomization::CustomizeChildren(TSharedRef<IPropertyH
 			}
 		}
 
+		TSharedPtr<IPropertyHandle> MetadataAssetClassHandle = PropertyHandle->GetChildHandle(GET_MEMBER_NAME_CHECKED(FAssemblyMetadataDesc, AssetClass));
+		MetadataAssetClassHandle->SetOnPropertyValueChanged(FSimpleDelegate::CreateRaw(this, &FCineAssemblyMetadataCustomization::ClearIncompatibleMetadataLink));
+
 		// Create a widget switcher that can display the appropriate widget based on the metadata type
 		TSharedPtr<IPropertyHandle> MetadataTypeHandle = PropertyHandle->GetChildHandle(GET_MEMBER_NAME_CHECKED(FAssemblyMetadataDesc, Type));
 
-		MetadataTypeHandle->SetOnPropertyValueChanged(FSimpleDelegate::CreateLambda([&MetadataDesc]()
+		MetadataTypeHandle->SetOnPropertyValueChanged(FSimpleDelegate::CreateLambda([this, &MetadataDesc]()
 			{
+				ClearIncompatibleMetadataLink();
+
 				if ((MetadataDesc.Type == ECineAssemblyMetadataType::String) ||
 					(MetadataDesc.Type == ECineAssemblyMetadataType::AssetPath) ||
 					(MetadataDesc.Type == ECineAssemblyMetadataType::CineAssembly))
@@ -183,7 +190,10 @@ void FCineAssemblyMetadataCustomization::CustomizeChildren(TSharedRef<IPropertyH
 
 		DefaultValueWidget->AddSlot((uint8)ECineAssemblyMetadataType::String)
 			[
-				MakeStringDefaultValueWidget(MetadataDesc, PropertyHandle)
+				SNew(SAssemblyMetadataLink, Schema->TemplateSequence, MetadataDesc.Key)
+					[
+						MakeStringDefaultValueWidget(MetadataDesc, PropertyHandle)
+					]
 			];
 
 		DefaultValueWidget->AddSlot((uint8)ECineAssemblyMetadataType::Bool)
@@ -243,30 +253,32 @@ void FCineAssemblyMetadataCustomization::CustomizeChildren(TSharedRef<IPropertyH
 
 		DefaultValueWidget->AddSlot((uint8)ECineAssemblyMetadataType::AssetPath)
 			[
-				SNew(SObjectPropertyEntryBox)
-					.ThumbnailPool(CustomizationUtils.GetThumbnailPool())
-					.AllowCreate(true)
-					.OnShouldFilterAsset_Lambda([&MetadataDesc](const FAssetData& InAssetData) -> bool
-						{
-							// Filter out all assets that do not match the selected Asset Class for this metadata struct
-							if (MetadataDesc.AssetClass.IsNull())
-							{
-								return true;
-							}
-							return (InAssetData.AssetClassPath != MetadataDesc.AssetClass.GetAssetPath());
-						})
-					.ObjectPath_Lambda([&MetadataDesc]() -> FString
-						{
-							if (const FString* Value = MetadataDesc.DefaultValue.TryGet<FString>())
-							{
-								return *Value;
-							}
-							return FString();
-						})
-					.OnObjectChanged_Lambda([&MetadataDesc](const FAssetData& InAssetData)
-						{
-							MetadataDesc.DefaultValue.Set<FString>(InAssetData.GetObjectPathString());
-						})
+				SNew(SAssemblyMetadataLink, Schema->TemplateSequence, MetadataDesc.Key)
+					[
+						SNew(SObjectPropertyEntryBox)
+							.ThumbnailPool(CustomizationUtils.GetThumbnailPool())
+							.AllowCreate(true)
+							.OnShouldFilterAsset_Lambda([&MetadataDesc](const FAssetData& InAssetData) -> bool
+								{
+									if (MetadataDesc.AssetClass.IsNull())
+									{
+										return false;
+									}
+									return (InAssetData.AssetClassPath != MetadataDesc.AssetClass.GetAssetPath());
+								})
+							.ObjectPath_Lambda([&MetadataDesc]() -> FString
+								{
+									if (const FString* Value = MetadataDesc.DefaultValue.TryGet<FString>())
+									{
+										return *Value;
+									}
+									return FString();
+								})
+							.OnObjectChanged_Lambda([&MetadataDesc](const FAssetData& InAssetData)
+								{
+									MetadataDesc.DefaultValue.Set<FString>(InAssetData.GetObjectPathString());
+								})
+					]
 			];
 
 		DefaultValueWidget->AddSlot((uint8)ECineAssemblyMetadataType::CineAssembly)
@@ -277,6 +289,11 @@ void FCineAssemblyMetadataCustomization::CustomizeChildren(TSharedRef<IPropertyH
 					.AllowCreate(true)
 					.OnShouldFilterAsset_Lambda([&MetadataDesc](const FAssetData& InAssetData) -> bool
 						{
+							if (!MetadataDesc.SchemaType.IsValid())
+							{
+								return false;
+							}
+
 							// Filter out all Cine Assembly assets that do not match the selected Schema Type for this metadata struct
 							const FAssetDataTagMapSharedView::FFindTagResult AssemblyType = InAssetData.TagsAndValues.FindTag(UCineAssembly::AssetRegistryTag_AssemblyType);
 							if (AssemblyType.IsSet())
@@ -313,6 +330,45 @@ void FCineAssemblyMetadataCustomization::CustomizeChildren(TSharedRef<IPropertyH
 	}
 }
 
+void FCineAssemblyMetadataCustomization::ClearIncompatibleMetadataLink()
+{
+	if (!Schema || !Schema->TemplateSequence || !Schema->AssemblyMetadata.IsValidIndex(ArrayIndex))
+	{
+		return;
+	}
+
+	const FAssemblyMetadataDesc& MetadataDesc = Schema->AssemblyMetadata[ArrayIndex];
+	const FGuid* LinkedGuid = Schema->TemplateSequence->MetadataLinks.Find(MetadataDesc.Key);
+	if (!LinkedGuid)
+	{
+		return;
+	}
+
+	if (MetadataDesc.Type != ECineAssemblyMetadataType::AssetPath && MetadataDesc.Type != ECineAssemblyMetadataType::String)
+	{
+		Schema->TemplateSequence->Modify();
+		Schema->TemplateSequence->MetadataLinks.Remove(MetadataDesc.Key);
+		return;
+	}
+
+	const UClass* NewAllowedClass = MetadataDesc.AssetClass.ResolveClass();
+	if (!NewAllowedClass)
+	{
+		// No class filter means any asset is compatible — keep the link
+		return;
+	}
+
+	if (const FAssemblyAssociatedAssetDesc* AssetDesc = Algo::FindBy(Schema->TemplateSequence->AssociatedAssets, *LinkedGuid, &FAssemblyAssociatedAssetDesc::AssetID))
+	{
+		UClass* AssetClass = AssetDesc->AssetClass.LoadSynchronous();
+		if (!AssetClass || !AssetClass->IsChildOf(NewAllowedClass))
+		{
+			Schema->TemplateSequence->Modify();
+			Schema->TemplateSequence->MetadataLinks.Remove(MetadataDesc.Key);
+		}
+	}
+}
+
 void FCineAssemblyMetadataCustomization::CustomizeKeyProperty(TSharedRef<IPropertyHandle> PropertyHandle, IDetailChildrenBuilder& ChildBuilder, FAssemblyMetadataDesc& MetadataDesc)
 {
 	FString ExistingKeyName;
diff --git a/Engine/Plugins/VirtualProduction/CinematicAssemblyTools/Source/CineAssemblyToolsEditor/Private/CineAssemblyMetadataCustomization.h b/Engine/Plugins/VirtualProduction/CinematicAssemblyTools/Source/CineAssemblyToolsEditor/Private/CineAssemblyMetadataCustomization.h
--- a/Engine/Plugins/VirtualProduction/CinematicAssemblyTools/Source/CineAssemblyToolsEditor/Private/CineAssemblyMetadataCustomization.h
+++ b/Engine/Plugins/VirtualProduction/CinematicAssemblyTools/Source/CineAssemblyToolsEditor/Private/CineAssemblyMetadataCustomization.h
@@ -26,6 +26,9 @@ private:
 	*/
 	TSharedRef<SWidget> MakeStringDefaultValueWidget(FAssemblyMetadataDesc& MetadataDesc, TSharedRef<IPropertyHandle> PropertyHandle);
 
+	/** Clears the metadata link if the linked asset's class is no longer compatible with the metadata field's AssetClass filter */
+	void ClearIncompatibleMetadataLink();
+
 	/** Customizes the metadata desc Key property */
 	void CustomizeKeyProperty(TSharedRef<IPropertyHandle> PropertyHandle, IDetailChildrenBuilder& ChildBuilder, FAssemblyMetadataDesc& MetadataDesc);
 
diff --git a/Engine/Plugins/VirtualProduction/CinematicAssemblyTools/Source/CineAssemblyToolsEditor/Private/UI/CineAssembly/SCineAssemblyAssetTreeView.cpp b/Engine/Plugins/VirtualProduction/CinematicAssemblyTools/Source/CineAssemblyToolsEditor/Private/UI/CineAssembly/SCineAssemblyAssetTreeView.cpp
--- a/Engine/Plugins/VirtualProduction/CinematicAssemblyTools/Source/CineAssemblyToolsEditor/Private/UI/CineAssembly/SCineAssemblyAssetTreeView.cpp
+++ b/Engine/Plugins/VirtualProduction/CinematicAssemblyTools/Source/CineAssemblyToolsEditor/Private/UI/CineAssembly/SCineAssemblyAssetTreeView.cpp
@@ -14,6 +14,7 @@
 #include "MovieSceneSubAssemblyTrack.h"
 #include "PropertyCustomizationHelpers.h"
 #include "Styling/SlateIconFinder.h"
+#include "UI/SAssemblyMetadataLink.h"
 #include "Widgets/Layout/SBorder.h"
 #include "Widgets/SBoxPanel.h"
 #include "Widgets/Text/STextBlock.h"
@@ -431,8 +432,7 @@ namespace UE::CineAssemblyTools::Private
 	{
 		if (UCineAssembly* Assembly = AssetItem->GetAssembly())
 		{
-			Assembly->Modify();
-			Assembly->AssociatedAssets.RemoveAll([&AssetItem](const FAssemblyAssociatedAssetDesc& AssetDesc) { return AssetDesc.AssetID == AssetItem->AssetID; });
+			Assembly->RemoveAssociatedAsset(AssetItem->AssetID);
 		}
 	}
 
@@ -800,9 +800,34 @@ TSharedRef<ITableRow> SCineAssemblyAssetTreeView::OnGenerateTreeRow(TSharedPtr<I
 						.Font(FCoreStyle::GetDefaultFontStyle("Italic", 10))
 						.ColorAndOpacity(FSlateColor(FLinearColor(1.0f, 1.0f, 1.0f, 0.3f)))
 				]
+
+			+ SHorizontalBox::Slot()
+				.AutoWidth()
+				.VAlign(VAlign_Center)
+				.Padding(4.0f, 0.0f, 0.0f, 0.0f)
+				[
+					MakeMetadataLinkButton(TreeItem)
+				]
 		];
 }
 
+TSharedRef<SWidget> SCineAssemblyAssetTreeView::MakeMetadataLinkButton(TSharedPtr<ICineAssemblyTreeItem> TreeItem)
+{
+	TSharedPtr<FAssociatedAssetItem> AssetItem = TreeItem->AsAssociatedAsset();
+	if (!AssetItem || bIsReadOnly)
+	{
+		return SNullWidget::NullWidget;
+	}
+
+	const FAssemblyAssociatedAssetDesc* AssetDesc = AssetItem->FindAssetDesc();
+	if (!AssetDesc)
+	{
+		return SNullWidget::NullWidget;
+	}
+
+	return SNew(SAssemblyMetadataLink, AssetItem->GetAssembly(), AssetDesc->AssetID);
+}
+
 void SCineAssemblyAssetTreeView::OnGetChildren(TSharedPtr<ICineAssemblyTreeItem> TreeItem, TArray<TSharedPtr<ICineAssemblyTreeItem>>& OutNodes)
 {
 	// Children are pre-sorted by SortTreeItems (assets first, then folders, alphabetical within each group)
diff --git a/Engine/Plugins/VirtualProduction/CinematicAssemblyTools/Source/CineAssemblyToolsEditor/Private/UI/CineAssembly/SCineAssemblyAssetTreeView.h b/Engine/Plugins/VirtualProduction/CinematicAssemblyTools/Source/CineAssemblyToolsEditor/Private/UI/CineAssembly/SCineAssemblyAssetTreeView.h
--- a/Engine/Plugins/VirtualProduction/CinematicAssemblyTools/Source/CineAssemblyToolsEditor/Private/UI/CineAssembly/SCineAssemblyAssetTreeView.h
+++ b/Engine/Plugins/VirtualProduction/CinematicAssemblyTools/Source/CineAssemblyToolsEditor/Private/UI/CineAssembly/SCineAssemblyAssetTreeView.h
@@ -139,6 +139,9 @@ private:
 	/** Recursively expands every folder in the tree view */
 	void ExpandTreeRecursive(const TSharedPtr<FCineAssemblyFolderItem>& FolderItem) const;
 
+	/** Creates the metadata link widget for associated asset items (returns SNullWidget for non-asset items or read-only trees) */
+	TSharedRef<SWidget> MakeMetadataLinkButton(TSharedPtr<ICineAssemblyTreeItem> TreeItem);
+
 	/** Begins a drag and drop event to drag an item out of the content tree view */
 	FReply OnTreeRowDragDetected(const FGeometry& MyGeometry, const FPointerEvent& MouseEvent);
 
diff --git a/Engine/Plugins/VirtualProduction/CinematicAssemblyTools/Source/CineAssemblyToolsEditor/Private/UI/CineAssembly/SCineAssemblySchemaWindow.cpp b/Engine/Plugins/VirtualProduction/CinematicAssemblyTools/Source/CineAssemblyToolsEditor/Private/UI/CineAssembly/SCineAssemblySchemaWindow.cpp
--- a/Engine/Plugins/VirtualProduction/CinematicAssemblyTools/Source/CineAssemblyToolsEditor/Private/UI/CineAssembly/SCineAssemblySchemaWindow.cpp
+++ b/Engine/Plugins/VirtualProduction/CinematicAssemblyTools/Source/CineAssemblyToolsEditor/Private/UI/CineAssembly/SCineAssemblySchemaWindow.cpp
@@ -764,6 +764,21 @@ void SCineAssemblySchemaWindow::OnSchemaPropertiesChanged(const FPropertyChanged
 			CineAssemblyNamingTokens->AddMetadataToken(MetadataDesc.Key);
 		}
 	}
+	else if (PropertyChangedEvent.GetMemberPropertyName() == GET_MEMBER_NAME_CHECKED(UCineAssemblySchema, AssemblyMetadata)	&& PropertyChangedEvent.ChangeType == EPropertyChangeType::ArrayRemove)
+	{
+		// Clean up MetadataLinks entries that reference metadata keys no longer in the schema
+		if (SchemaToConfigure->TemplateSequence)
+		{
+			SchemaToConfigure->TemplateSequence->Modify();
+			for (auto It = SchemaToConfigure->TemplateSequence->MetadataLinks.CreateIterator(); It; ++It)
+			{
+				if (!Algo::ContainsBy(SchemaToConfigure->AssemblyMetadata, It->Key, &FAssemblyMetadataDesc::Key))
+				{
+					It.RemoveCurrent();
+				}
+			}
+		}
+	}
 	else if (PropertyName == GET_MEMBER_NAME_CHECKED(UCineAssemblySchema, ThumbnailImage))
 	{
 		// Update the schema's thumbnail brush in the style set to use the new thumbnail image
diff --git a/Engine/Plugins/VirtualProduction/CinematicAssemblyTools/Source/CineAssemblyToolsEditor/Private/UI/SAssemblyMetadataLink.cpp b/Engine/Plugins/VirtualProduction/CinematicAssemblyTools/Source/CineAssemblyToolsEditor/Private/UI/SAssemblyMetadataLink.cpp
new file mode 100644
--- /dev/null
+++ b/Engine/Plugins/VirtualProduction/CinematicAssemblyTools/Source/CineAssemblyToolsEditor/Private/UI/SAssemblyMetadataLink.cpp
@@ -0,0 +1,384 @@
+// Copyright Epic Games, Inc. All Rights Reserved.
+
+#include "SAssemblyMetadataLink.h"
+
+#include "Algo/Contains.h"
+#include "Algo/Find.h"
+#include "CineAssembly.h"
+#include "CineAssemblySchema.h"
+#include "Framework/MultiBox/MultiBoxBuilder.h"
+#include "Styling/StyleColors.h"
+#include "Widgets/Images/SImage.h"
+#include "Widgets/Input/SComboButton.h"
+#include "Widgets/Layout/SWidgetSwitcher.h"
+#include "Widgets/Text/STextBlock.h"
+
+#define LOCTEXT_NAMESPACE "SAssemblyMetadataLink"
+
+void SAssemblyMetadataLink::Construct(const FArguments& InArgs, UCineAssembly* InAssembly, FGuid InAssetID)
+{
+	// The widget requires a valid input assembly to be operational
+	if (!ensure(InAssembly))
+	{
+		return;
+	}
+
+	// Asset Mode should not provide an unlinked value widget
+	ensure(InArgs._UnlinkedValueWidget.Widget == SNullWidget::NullWidget);
+
+	WeakAssembly = InAssembly;
+	AssociatedAssetID = InAssetID;
+	bIsAssetMode = true;
+
+	if (const FAssemblyAssociatedAssetDesc* AssetDesc = Algo::FindBy(InAssembly->AssociatedAssets, AssociatedAssetID, &FAssemblyAssociatedAssetDesc::AssetID))
+	{
+		CachedAssetClass = AssetDesc->AssetClass.Get();
+	}
+
+	ConstructInternal(InArgs);
+}
+
+void SAssemblyMetadataLink::Construct(const FArguments& InArgs, UCineAssembly* InAssembly, const FString& InMetadataKey)
+{
+	// The widget requires a valid input assembly to be operational
+	if (!ensure(InAssembly))
+	{
+		return;
+	}
+
+	WeakAssembly = InAssembly;
+	MetadataKey = InMetadataKey;
+	bIsAssetMode = false;
+
+	if (const UCineAssemblySchema* Schema = GetSchema())
+	{
+		if (const FAssemblyMetadataDesc* MetadataDesc = Algo::FindBy(Schema->AssemblyMetadata, MetadataKey, &FAssemblyMetadataDesc::Key))
+		{
+			switch(MetadataDesc->Type)
+			{
+			case ECineAssemblyMetadataType::AssetPath:
+				CachedAssetClass = MetadataDesc->AssetClass.ResolveClass();
+				break;
+			default:
+				// For all other metadata types, the AssetClass is not used to filter out associated assets that are compatible for linking
+				CachedAssetClass = nullptr;
+				break;
+			}
+		}
+	}
+
+	ConstructInternal(InArgs);
+}
+
+void SAssemblyMetadataLink::ConstructInternal(const FArguments& InArgs)
+{
+	const bool bHasValueWidget = InArgs._UnlinkedValueWidget.Widget != SNullWidget::NullWidget;
+
+	TSharedRef<SHorizontalBox> Content = SNew(SHorizontalBox);
+
+	constexpr int32 UnlinkedWidgetIndex = 0;
+	constexpr int32 LinkedWidgetIndex = 1;
+
+	if (bHasValueWidget && !bIsAssetMode)
+	{
+		Content->AddSlot()
+			.FillWidth(1.0f)
+			.Padding(0.0f, 0.0f, 4.0f, 0.0f)
+			[
+				SNew(SWidgetSwitcher)
+					.WidgetIndex_Lambda([this]() { return HasAnyLink() ? LinkedWidgetIndex : UnlinkedWidgetIndex; })
+
+				+ SWidgetSwitcher::Slot()
+					[
+						InArgs._UnlinkedValueWidget.Widget
+					]
+
+				+ SWidgetSwitcher::Slot()
+					.VAlign(VAlign_Center)
+					[
+						SNew(STextBlock)
+							.Text(this, &SAssemblyMetadataLink::GetLinkedAssetName)
+							.Font(FAppStyle::GetFontStyle(TEXT("PropertyWindow.NormalFont")))
+					]
+			];
+	}
+
+	Content->AddSlot()
+		.AutoWidth()
+		.VAlign(VAlign_Center)
+		[
+			SNew(SComboButton)
+				.HasDownArrow(false)
+				.ContentPadding(FMargin(2.0f))
+				.ButtonStyle(&FAppStyle::Get().GetWidgetStyle<FButtonStyle>("SimpleButton"))
+				.ToolTipText(bIsAssetMode
+					? LOCTEXT("AssetModeLinkTooltip", "Link this asset to a metadata field so it is auto-populated when the assembly is created")
+					: LOCTEXT("FieldModeLinkTooltip", "Link to an associated asset so this field is auto-populated when the assembly is created"))
+				.Visibility(this, &SAssemblyMetadataLink::GetLinkButtonVisibility)
+				.OnGetMenuContent(this, &SAssemblyMetadataLink::BuildMenuContent)
+				.ButtonContent()
+				[
+					SNew(SImage)
+						.Image(FAppStyle::Get().GetBrush("Icons.Link"))
+						.ColorAndOpacity(this, &SAssemblyMetadataLink::GetLinkButtonColor)
+				]
+		];
+
+	ChildSlot
+	[
+		Content
+	];
+}
+
+TSharedRef<SWidget> SAssemblyMetadataLink::BuildMenuContent()
+{
+	UCineAssembly* Assembly = WeakAssembly.Get();
+	const UCineAssemblySchema* Schema = GetSchema();
+	if (!Assembly || !Schema)
+	{
+		return SNullWidget::NullWidget;
+	}
+
+	constexpr bool bCloseAfterMenuSelection = true;
+	FMenuBuilder MenuBuilder(bCloseAfterMenuSelection, nullptr);
+
+	if (bIsAssetMode)
+	{
+		// Asset Mode: list compatible metadata fields as checkboxes (an associated asset can be linked to multiple metadata fields)
+		for (const FAssemblyMetadataDesc& MetadataDesc : Schema->AssemblyMetadata)
+		{
+			if (!IsMetadataDescCompatible(MetadataDesc))
+			{
+				continue;
+			}
+
+			const FText Tooltip = IsAssetLinked(MetadataDesc.Key, AssociatedAssetID)
+				? FText::Format(LOCTEXT("UnlinkMetaTooltip", "Unlink from the '{0}' metadata field"), FText::FromString(MetadataDesc.Key))
+				: FText::Format(LOCTEXT("LinkMetaTooltip", "Link to the '{0}' metadata field"), FText::FromString(MetadataDesc.Key));
+
+			MenuBuilder.AddMenuEntry(
+				FText::FromString(MetadataDesc.Key),
+				Tooltip,
+				FSlateIcon(),
+				FUIAction(
+					FExecuteAction::CreateSP(this, &SAssemblyMetadataLink::ToggleMetadataLink, MetadataDesc.Key, AssociatedAssetID),
+					FCanExecuteAction(),
+					FIsActionChecked::CreateSP(this, &SAssemblyMetadataLink::IsAssetLinked, MetadataDesc.Key, AssociatedAssetID)
+				),
+				NAME_None,
+				EUserInterfaceActionType::Check
+			);
+		}
+	}
+	else if (IsMetadataKeyValid())
+	{
+		// Metadata Mode: list compatible associated assets as radio buttons (a metadata field can only be linked to one associated asset)
+		for (const FAssemblyAssociatedAssetDesc& AssetDesc : Assembly->AssociatedAssets)
+		{
+			if (!IsAssetClassCompatible(AssetDesc.AssetClass.Get(), CachedAssetClass.Get()))
+			{
+				continue;
+			}
+
+			const FText Tooltip = IsAssetLinked(MetadataKey, AssetDesc.AssetID)
+				? FText::Format(LOCTEXT("UnlinkAssetTooltip", "Unlink from '{0}'"), FText::FromString(AssetDesc.AssetName.Template))
+				: FText::Format(LOCTEXT("LinkAssetTooltip", "Link to '{0}'"), FText::FromString(AssetDesc.AssetName.Template));
+
+			MenuBuilder.AddMenuEntry(
+				FText::FromString(AssetDesc.AssetName.Template),
+				Tooltip,
+				FSlateIcon(),
+				FUIAction(
+					FExecuteAction::CreateSP(this, &SAssemblyMetadataLink::ToggleMetadataLink, MetadataKey, AssetDesc.AssetID),
+					FCanExecuteAction(),
+					FIsActionChecked::CreateSP(this, &SAssemblyMetadataLink::IsAssetLinked, MetadataKey, AssetDesc.AssetID)
+				),
+				NAME_None,
+				EUserInterfaceActionType::RadioButton
+			);
+		}
+	}
+
+	return MenuBuilder.MakeWidget();
+}
+
+void SAssemblyMetadataLink::ToggleMetadataLink(FString InMetadataKey, FGuid InAssetID)
+{
+	if (UCineAssembly* Assembly = WeakAssembly.Get())
+	{
+		Assembly->Modify();
+		if (IsAssetLinked(InMetadataKey, InAssetID))
+		{
+			Assembly->MetadataLinks.Remove(InMetadataKey);
+		}
+		else
+		{
+			Assembly->MetadataLinks.Add(InMetadataKey, InAssetID);
+		}
+	}
+}
+
+bool SAssemblyMetadataLink::IsAssetLinked(FString InMetadataKey, FGuid InAssetID) const
+{
+	if (UCineAssembly* Assembly = WeakAssembly.Get())
+	{
+		const FGuid* LinkedID = Assembly->MetadataLinks.Find(InMetadataKey);
+		return LinkedID && *LinkedID == InAssetID;
+	}
+	return false;
+}
+
+bool SAssemblyMetadataLink::HasAnyLink() const
+{
+	UCineAssembly* Assembly = WeakAssembly.Get();
+	if (!Assembly)
+	{
+		return false;
+	}
+
+	if (bIsAssetMode)
+	{
+		// Check if any metadata key in the map points to this asset
+		for (const TPair<FString, FGuid>& MetadataLink : Assembly->MetadataLinks)
+		{
+			if (MetadataLink.Value == AssociatedAssetID)
+			{
+				return true;
+			}
+		}
+		return false;
+	}
+	else if (IsMetadataKeyValid())
+	{
+		return Assembly->MetadataLinks.Contains(MetadataKey);
+	}
+
+	return false;
+}
+
+EVisibility SAssemblyMetadataLink::GetLinkButtonVisibility() const
+{
+	UCineAssembly* Assembly = WeakAssembly.Get();
+	const UCineAssemblySchema* Schema = GetSchema();
+	if (!Assembly || !Schema)
+	{
+		return EVisibility::Collapsed;
+	}
+
+	if (bIsAssetMode)
+	{
+		// Visible if there are any compatible metadata fields
+		const bool bHasCompatibleField = Schema->AssemblyMetadata.ContainsByPredicate([this](const FAssemblyMetadataDesc& MetadataDesc) { return IsMetadataDescCompatible(MetadataDesc); });
+		return bHasCompatibleField ? EVisibility::Visible : EVisibility::Collapsed;
+	}
+	else
+	{
+		if (!IsMetadataKeyValid())
+		{
+			return EVisibility::Collapsed;
+		}
+
+		const bool bHasCompatibleAsset = Assembly->AssociatedAssets.ContainsByPredicate([this](const FAssemblyAssociatedAssetDesc& AssetDesc) { return IsAssetClassCompatible(AssetDesc.AssetClass.Get(), CachedAssetClass.Get()); });
+		return bHasCompatibleAsset ? EVisibility::Visible : EVisibility::Collapsed;
+	}
+}
+
+FSlateColor SAssemblyMetadataLink::GetLinkButtonColor() const
+{
+	return HasAnyLink() ? FStyleColors::White : FSlateColor::UseSubduedForeground();
+}
+
+FText SAssemblyMetadataLink::GetLinkedAssetName() const
+{
+	if (bIsAssetMode)
+	{
+		// In Asset Mode, there is no extra unlinked/linked slot to display any text for
+		return FText::GetEmpty();
+	}
+
+	if (!IsMetadataKeyValid())
+	{
+		return FText::GetEmpty();
+	}
+
+	if (UCineAssembly* Assembly = WeakAssembly.Get())
+	{
+		// In Metadata Mode, show the linked asset's template name
+		if (const FGuid* LinkedID = Assembly->MetadataLinks.Find(MetadataKey))
+		{
+			if (const FAssemblyAssociatedAssetDesc* AssetDesc = Algo::FindBy(Assembly->AssociatedAssets, *LinkedID, &FAssemblyAssociatedAssetDesc::AssetID))
+			{
+				return FText::FromString(AssetDesc->AssetName.Template);
+			}
+		}
+
+		return LOCTEXT("LinkedAssetNotFound", "(Linked asset not found)");
+	}
+
+	return FText::GetEmpty();
+}
+
+const UCineAssemblySchema* SAssemblyMetadataLink::GetSchema() const
+{
+	UCineAssembly* Assembly = WeakAssembly.Get();
+	if (!Assembly)
+	{
+		return nullptr;
+	}
+
+	// Try the outer chain first (for template sequences owned by schemas), then fall back to GetSchema()
+	if (const UCineAssemblySchema* Schema = Assembly->GetTypedOuter<UCineAssemblySchema>())
+	{
+		return Schema;
+	}
+	return Assembly->GetSchema();
+}
+
+bool SAssemblyMetadataLink::IsMetadataKeyValid() const
+{
+	if (bIsAssetMode)
+	{
+		return false;
+	}
+
+	const UCineAssemblySchema* Schema = GetSchema();
+	if (!Schema)
+	{
+		return false;
+	}
+	return Algo::ContainsBy(Schema->AssemblyMetadata, MetadataKey, &FAssemblyMetadataDesc::Key);
+}
+
+bool SAssemblyMetadataLink::IsAssetClassCompatible(const UClass* InAssetClass, const UClass* InFilterClass) const
+{
+	if (!InFilterClass)
+	{
+		return true;
+	}
+	return InAssetClass && InAssetClass->IsChildOf(InFilterClass);
+}
+
+bool SAssemblyMetadataLink::IsMetadataDescCompatible(const FAssemblyMetadataDesc& InMetadataDesc) const
+{
+	switch(InMetadataDesc.Type)
+	{
+	case ECineAssemblyMetadataType::String:
+		return true;
+
+	case ECineAssemblyMetadataType::AssetPath:
+		return IsAssetClassCompatible(CachedAssetClass.Get(), InMetadataDesc.AssetClass.ResolveClass());
+
+	case ECineAssemblyMetadataType::Bool:
+	case ECineAssemblyMetadataType::Integer:
+	case ECineAssemblyMetadataType::Float:
+	case ECineAssemblyMetadataType::CineAssembly:
+		return false;
+
+	default:
+		checkNoEntry();
+	}
+
+	return false;
+}
+
+#undef LOCTEXT_NAMESPACE
diff --git a/Engine/Plugins/VirtualProduction/CinematicAssemblyTools/Source/CineAssemblyToolsEditor/Private/UI/SAssemblyMetadataLink.h b/Engine/Plugins/VirtualProduction/CinematicAssemblyTools/Source/CineAssemblyToolsEditor/Private/UI/SAssemblyMetadataLink.h
new file mode 100644
--- /dev/null
+++ b/Engine/Plugins/VirtualProduction/CinematicAssemblyTools/Source/CineAssemblyToolsEditor/Private/UI/SAssemblyMetadataLink.h
@@ -0,0 +1,94 @@
+// Copyright Epic Games, Inc. All Rights Reserved.
+
+#pragma once
+
+#include "Widgets/SCompoundWidget.h"
+
+class UCineAssembly;
+class UCineAssemblySchema;
+struct FAssemblyAssociatedAssetDesc;
+struct FAssemblyMetadataDesc;
+
+/**
+ * A composite widget that handles the linking between a CineAssembly's associated assets and metadata fields.
+ *
+ * Operates in one of two modes:
+ *     Asset Mode:    Handles the logic to link an associated asset to one or more metadata fields
+ *     Metadata Mode: Handles the logic to link a metadata field to an associated asset
+ *     
+ * In Metadata Mode, this widget wraps a value widget to display when the field is not linked, 
+ * and handles switching to a read-only textbox displaying the name of the linked asset when it is linked.
+ */
+class SAssemblyMetadataLink : public SCompoundWidget
+{
+public:
+	SLATE_BEGIN_ARGS(SAssemblyMetadataLink)
+		{}
+
+		/** Optional value widget to show when unlinked. Only applicable to Metadata Mode */
+		SLATE_DEFAULT_SLOT(FArguments, UnlinkedValueWidget)
+
+	SLATE_END_ARGS()
+
+	/** Asset Mode */
+	void Construct(const FArguments& InArgs, UCineAssembly* InAssembly, FGuid InAssetID);
+
+	/** Metadata Mode */
+	void Construct(const FArguments& InArgs, UCineAssembly* InAssembly, const FString& InMetadataKey);
+
+private:
+	/** Shared widget construction for Asset and Metadata modes */
+	void ConstructInternal(const FArguments& InArgs);
+
+	/** Builds the dropdown menu content showing the link options */
+	TSharedRef<SWidget> BuildMenuContent();
+
+	/** Toggles a link between a metadata key and an asset GUID */
+	void ToggleMetadataLink(FString InMetadataKey, FGuid InAssetID);
+
+	/** Whether the input metadata key and associated asset are currently linked */
+	bool IsAssetLinked(FString InMetadataKey, FGuid InAssetID) const;
+
+	/**
+	 * AssetMode: Whether the tracked AssociatedAssetID is linked to any metadata field
+	 * MetadataMode: Whether the tracked MetadataKey is linked to any associated asset
+	 */
+	bool HasAnyLink() const;
+
+	/** Determines the visibility of the link button */
+	EVisibility GetLinkButtonVisibility() const;
+
+	/** Returns the link button icon color based on link status */
+	FSlateColor GetLinkButtonColor() const;
+
+	/** Returns the template name of the linked target */
+	FText GetLinkedAssetName() const;
+
+	/** Returns the schema that defines the assembly's metadata fields */
+	const UCineAssemblySchema* GetSchema() const;
+
+	/** Whether the cached metadata key is still in the schema's list of metadata fields */
+	bool IsMetadataKeyValid() const;
+
+	/** Whether the AssetClass is compatible with the FilterClass. Always returns true if FilterClass is nullptr. */
+	bool IsAssetClassCompatible(const UClass* InAssetClass, const UClass* InFilterClass) const;
+
+	/** Whether the input metadata desc is compatible for linking with this asset */
+	bool IsMetadataDescCompatible(const FAssemblyMetadataDesc& InMetadataDesc) const;
+
+private:
+	/** The assembly whose associated assets and schema metadata are being linked */
+	TWeakObjectPtr<UCineAssembly> WeakAssembly;
+
+	/** Asset Mode: the associated asset being linked */
+	FGuid AssociatedAssetID;
+
+	/** Metadata Mode: the metadata field key */
+	FString MetadataKey;
+
+	/** Whether this widget is in Asset Mode (true) or Metadata Mode (false) */
+	bool bIsAssetMode = false;
+
+	/** Cached asset class — in Asset Mode: the associated asset's class; in Metadata Mode: the metadata field's allowed class */
+	TWeakObjectPtr<const UClass> CachedAssetClass;
+};
