ShotManagement: Add metadata linking support for SubAssemblies. In the Schema editor, users can now link a SubAssembly that will get created to populate the value of a String, AssetPath, or CineAssembly type metadata field.
#rnx
#Synced-CL 51715153

diff --git a/Engine/Plugins/VirtualProduction/CinematicAssemblyTools/Source/CineAssemblyTools/Private/CineAssembly.cpp b/Engine/Plugins/VirtualProduction/CinematicAssemblyTools/Source/CineAssemblyTools/Private/CineAssembly.cpp
--- a/Engine/Plugins/VirtualProduction/CinematicAssemblyTools/Source/CineAssemblyTools/Private/CineAssembly.cpp
+++ b/Engine/Plugins/VirtualProduction/CinematicAssemblyTools/Source/CineAssemblyTools/Private/CineAssembly.cpp
@@ -399,6 +399,65 @@ void UCineAssembly::RemoveAssociatedAsset(const FGuid& InAssetID)
 	}
 }
 
+void UCineAssembly::ValidateMetadataLinks()
+{
+	if (MetadataLinks.IsEmpty())
+	{
+		return;
+	}
+
+	const UCineAssemblySchema* Schema = BaseSchema.Get() ? BaseSchema.Get() : GetTypedOuter<UCineAssemblySchema>();
+	if (!Schema)
+	{
+		return;
+	}
+
+	// Gather all valid GUIDs from associated assets
+	TSet<FGuid> ValidGUIDs;
+	for (const FAssemblyAssociatedAssetDesc& AssetDesc : AssociatedAssets)
+	{
+		ValidGUIDs.Add(AssetDesc.AssetID);
+	}
+
+	// Gather valid GUIDs from SubAssembly sections in the MovieScene
+	if (UMovieScene* InMovieScene = GetMovieScene())
+	{
+		for (UMovieSceneSection* Section : InMovieScene->GetAllSections())
+		{
+			if (UMovieSceneSubAssemblySection* SubAssemblySection = Cast<UMovieSceneSubAssemblySection>(Section))
+			{
+				ValidGUIDs.Add(SubAssemblySection->GetSectionID());
+			}
+		}
+	}
+
+	// Gather valid GUIDs from managed SubAssemblies
+	for (UMovieSceneSubSection* SubSection : SubAssemblies)
+	{
+		if (UCineAssembly* SubAssembly = Cast<UCineAssembly>(SubSection ? SubSection->GetSequence() : nullptr))
+		{
+			ValidGUIDs.Add(SubAssembly->GetAssemblyGuid());
+		}
+	}
+
+	// Gather valid metadata keys from the schema
+	TSet<FString> ValidKeys;
+	for (const FAssemblyMetadataDesc& MetadataDesc : Schema->AssemblyMetadata)
+	{
+		ValidKeys.Add(MetadataDesc.Key);
+	}
+
+	// Remove entries with stale keys or stale GUIDs
+	for (auto It = MetadataLinks.CreateIterator(); It; ++It)
+	{
+		if (!ValidKeys.Contains(It->Key) || !ValidGUIDs.Contains(It->Value))
+		{
+			Modify();
+			It.RemoveCurrent();
+		}
+	}
+}
+
 void UCineAssembly::DuplicateManagedSubAssemblies()
 {
 	if (!MovieScene)
@@ -492,6 +551,15 @@ void UCineAssembly::ConvertSubAssemblyTracks()
 
 					NewSubSection->SetSequence(NewSubAssembly);
 					SubAssemblies.Add(NewSubSection);
+
+					// Remap any MetadataLinks entries that reference this section's SectionID to the new SubAssembly's AssemblyGuid
+					for (TPair<FString, FGuid>& MetadataLink : MetadataLinks)
+					{
+						if (MetadataLink.Value == SubAssemblySection->GetSectionID())
+						{
+							MetadataLink.Value = NewSubAssembly->GetAssemblyGuid();
+						}
+					}
 				}
 			}
 		}
diff --git a/Engine/Plugins/VirtualProduction/CinematicAssemblyTools/Source/CineAssemblyTools/Private/MovieSceneSubAssemblySection.cpp b/Engine/Plugins/VirtualProduction/CinematicAssemblyTools/Source/CineAssemblyTools/Private/MovieSceneSubAssemblySection.cpp
--- a/Engine/Plugins/VirtualProduction/CinematicAssemblyTools/Source/CineAssemblyTools/Private/MovieSceneSubAssemblySection.cpp
+++ b/Engine/Plugins/VirtualProduction/CinematicAssemblyTools/Source/CineAssemblyTools/Private/MovieSceneSubAssemblySection.cpp
@@ -18,6 +18,57 @@ UMovieSceneSubAssemblySection::UMovieSceneSubAssemblySection(const FObjectInitia
 {
 }
 
+void UMovieSceneSubAssemblySection::PostInitProperties()
+{
+	Super::PostInitProperties();
+
+	if (!SectionID.IsValid())
+	{
+		SectionID = FGuid::NewGuid();
+	}
+}
+
+void UMovieSceneSubAssemblySection::PostDuplicate(bool bDuplicateForPIE)
+{
+	Super::PostDuplicate(bDuplicateForPIE);
+
+	// Assembly metadata linking depends on SubAssembly GUIDs being stable during Assembly initialization.
+	// During initialization, the Assembly will duplicate its MovieScene from a template sequence, and we want the duplicate section to have the same GUID.
+	// The transient duplicate check strongly suggests that this duplication is happening during Assembly initialization.
+	if (!bDuplicateForPIE && GetOutermost() != GetTransientPackage())
+	{
+		SectionID = FGuid::NewGuid();
+	}
+}
+
+void UMovieSceneSubAssemblySection::Serialize(FArchive& Ar)
+{
+	Super::Serialize(Ar);
+
+	if (Ar.IsLoading() && !SectionID.IsValid())
+	{
+		SectionID = FGuid::NewGuid();
+	}
+}
+
+FGuid UMovieSceneSubAssemblySection::GetSectionID() const
+{
+	return SectionID;
+}
+
+const UCineAssemblySchema* UMovieSceneSubAssemblySection::GetTemplateSchema() const
+{
+	if (const UCineAssemblySchema* Schema = Cast<UCineAssemblySchema>(AssemblyTemplate))
+	{
+		return Schema;
+	}
+	if (const UCineAssembly* Assembly = Cast<UCineAssembly>(AssemblyTemplate))
+	{
+		return Assembly->GetSchema();
+	}
+	return nullptr;
+}
+
 bool UMovieSceneSubAssemblySection::IsReferenceSection() const
 { 
 	return SectionType == ESubAssemblySectionType::Reference; 
diff --git a/Engine/Plugins/VirtualProduction/CinematicAssemblyTools/Source/CineAssemblyTools/Public/CineAssembly.h b/Engine/Plugins/VirtualProduction/CinematicAssemblyTools/Source/CineAssemblyTools/Public/CineAssembly.h
--- a/Engine/Plugins/VirtualProduction/CinematicAssemblyTools/Source/CineAssemblyTools/Public/CineAssembly.h
+++ b/Engine/Plugins/VirtualProduction/CinematicAssemblyTools/Source/CineAssemblyTools/Public/CineAssembly.h
@@ -223,6 +223,9 @@ public:
 	/** Removes an associated asset by its AssetID, and cleans up any MetadataLinks that referenced it */
 	CINEASSEMBLYTOOLS_API void RemoveAssociatedAsset(const FGuid& InAssetID);
 
+	/** Removes stale entries from MetadataLinks whose keys or values no longer match any metadata field, associated asset, or SubAssembly section */
+	CINEASSEMBLYTOOLS_API void ValidateMetadataLinks();
+
 	/** Recursively duplicates all managed SubAssemblies, replacing each with a new transient copy */
 	CINEASSEMBLYTOOLS_API void DuplicateManagedSubAssemblies();
 
@@ -302,7 +305,7 @@ public:
 	UPROPERTY()
 	TArray<FAssemblyAssociatedAssetDesc> AssociatedAssets;
 
-	/** Links metadata keys to the associated assets to auto-populate the value of the linked metadata field when the asset is created */
+	/** Links metadata keys to AssociatedAssets and SubAssemblies to auto-populate the value of the linked metadata field when the asset is created */
 	UPROPERTY()
 	TMap<FString, FGuid> MetadataLinks;
 
diff --git a/Engine/Plugins/VirtualProduction/CinematicAssemblyTools/Source/CineAssemblyTools/Public/MovieSceneSubAssemblySection.h b/Engine/Plugins/VirtualProduction/CinematicAssemblyTools/Source/CineAssemblyTools/Public/MovieSceneSubAssemblySection.h
--- a/Engine/Plugins/VirtualProduction/CinematicAssemblyTools/Source/CineAssemblyTools/Public/MovieSceneSubAssemblySection.h
+++ b/Engine/Plugins/VirtualProduction/CinematicAssemblyTools/Source/CineAssemblyTools/Public/MovieSceneSubAssemblySection.h
@@ -6,6 +6,8 @@
 
 #include "MovieSceneSubAssemblySection.generated.h"
 
+class UCineAssemblySchema;
+
 #define UE_API CINEASSEMBLYTOOLS_API
 
 /** Indicates how a SubAssemblySection will be used */
@@ -28,11 +30,20 @@ class UMovieSceneSubAssemblySection : public UMovieSceneCinematicShotSection
 public:
 	UMovieSceneSubAssemblySection(const FObjectInitializer& ObjectInitializer);
 
-#if WITH_EDITOR
 	//~ Begin UObject interface
+	virtual void PostInitProperties() override;
+	virtual void PostDuplicate(bool bDuplicateForPIE) override;
+	virtual void Serialize(FArchive& Ar) override;
+#if WITH_EDITOR
 	virtual void PostEditChangeProperty(FPropertyChangedEvent& PropertyChangedEvent) override;
-	//~ End UObject interface
 #endif //WITH_EDITOR
+	//~ End UObject interface
+
+	/** Get the unique identifier for this section */
+	UE_API FGuid GetSectionID() const;
+
+	/** Returns the schema associated with this section's template. Returns nullptr if no schema can be determined. */
+	UE_API const UCineAssemblySchema* GetTemplateSchema() const;
 
 	/** Returns true if this section is used as a reference to an existing sequence */
 	UE_API bool IsReferenceSection() const;
@@ -89,6 +100,10 @@ protected:
 	/** The object that will be used as a template to create a new sequence in this section. Can be a Level Sequence, Cine Assembly, or Schema */
 	UPROPERTY(EditAnywhere, Category = "Sequence")
 	TObjectPtr<UObject> AssemblyTemplate;
+
+	/** Unique identifier for this section */
+	UPROPERTY()
+	FGuid SectionID;
 };
 
 #undef UE_API
diff --git a/Engine/Plugins/VirtualProduction/CinematicAssemblyTools/Source/CineAssemblyToolsEditor/Private/CineAssemblyFactory.cpp b/Engine/Plugins/VirtualProduction/CinematicAssemblyTools/Source/CineAssemblyToolsEditor/Private/CineAssemblyFactory.cpp
--- a/Engine/Plugins/VirtualProduction/CinematicAssemblyTools/Source/CineAssemblyToolsEditor/Private/CineAssemblyFactory.cpp
+++ b/Engine/Plugins/VirtualProduction/CinematicAssemblyTools/Source/CineAssemblyToolsEditor/Private/CineAssemblyFactory.cpp
@@ -39,7 +39,7 @@ namespace UE::CineAssemblyTools::Private
 		return nullptr;
 	}
 
-	/** Updates metadata fields on the assembly that are linked to associated assets */
+	/** Set the value of metadata fields that are linked to the input Assembly's AssociatedAssets and SubAssemblies */
 	void PopulateLinkedMetadata(UCineAssembly* Assembly)
 	{
 		for (const TPair<FString, FGuid>& MetadataLink : Assembly->MetadataLinks)
@@ -49,12 +49,26 @@ namespace UE::CineAssemblyTools::Private
 				continue;
 			}
 
-			// Find the associated asset that matches the linked GUID
+			// Check the AssociatedAssets for an asset matching the linked ID
 			if (const FAssemblyAssociatedAssetDesc* AssetDesc = Algo::FindBy(Assembly->AssociatedAssets, MetadataLink.Value, &FAssemblyAssociatedAssetDesc::AssetID))
 			{
-				if (!AssetDesc->CreatedAsset.IsNull())
+				if (UObject* CreatedAsset = AssetDesc->CreatedAsset.Get())
 				{
-					Assembly->SetMetadataAsString(MetadataLink.Key, AssetDesc->CreatedAsset.ToSoftObjectPath().ToString());
+					Assembly->SetMetadataAsString(MetadataLink.Key, CreatedAsset->GetPathName());
+				}
+				continue;
+			}
+
+			// Check the SubAssemblies for an assembly matching the linked ID
+			for (UMovieSceneSubSection* SubSection : Assembly->SubAssemblies)
+			{
+				if (UCineAssembly* SubAssembly = Cast<UCineAssembly>(SubSection ? SubSection->GetSequence() : nullptr))
+				{
+					if (SubAssembly->GetAssemblyGuid() == MetadataLink.Value)
+					{
+						Assembly->SetMetadataAsString(MetadataLink.Key, SubAssembly->GetPathName());
+						break;
+					}
 				}
 			}
 		}
diff --git a/Engine/Plugins/VirtualProduction/CinematicAssemblyTools/Source/CineAssemblyToolsEditor/Private/CineAssemblyMetadataCustomization.cpp b/Engine/Plugins/VirtualProduction/CinematicAssemblyTools/Source/CineAssemblyToolsEditor/Private/CineAssemblyMetadataCustomization.cpp
--- a/Engine/Plugins/VirtualProduction/CinematicAssemblyTools/Source/CineAssemblyToolsEditor/Private/CineAssemblyMetadataCustomization.cpp
+++ b/Engine/Plugins/VirtualProduction/CinematicAssemblyTools/Source/CineAssemblyToolsEditor/Private/CineAssemblyMetadataCustomization.cpp
@@ -10,6 +10,8 @@
 #include "CineAssemblyToolsStyle.h"
 #include "IDetailChildrenBuilder.h"
 #include "IPropertyUtilities.h"
+#include "MovieScene.h"
+#include "MovieSceneSubAssemblySection.h"
 #include "NamingTokensEngineSubsystem.h"
 #include "PropertyCustomizationHelpers.h"
 #include "SNamingTokensEditableTextBox.h"
@@ -146,14 +148,17 @@ void FCineAssemblyMetadataCustomization::CustomizeChildren(TSharedRef<IPropertyH
 		}
 
 		TSharedPtr<IPropertyHandle> MetadataAssetClassHandle = PropertyHandle->GetChildHandle(GET_MEMBER_NAME_CHECKED(FAssemblyMetadataDesc, AssetClass));
-		MetadataAssetClassHandle->SetOnPropertyValueChanged(FSimpleDelegate::CreateRaw(this, &FCineAssemblyMetadataCustomization::ClearIncompatibleMetadataLink));
+		MetadataAssetClassHandle->SetOnPropertyValueChanged(FSimpleDelegate::CreateRaw(this, &FCineAssemblyMetadataCustomization::ValidateMetadataLink));
+
+		TSharedPtr<IPropertyHandle> MetadataSchemaTypeHandle = PropertyHandle->GetChildHandle(GET_MEMBER_NAME_CHECKED(FAssemblyMetadataDesc, SchemaType));
+		MetadataSchemaTypeHandle->SetOnPropertyValueChanged(FSimpleDelegate::CreateRaw(this, &FCineAssemblyMetadataCustomization::ValidateMetadataLink));
 
 		// Create a widget switcher that can display the appropriate widget based on the metadata type
 		TSharedPtr<IPropertyHandle> MetadataTypeHandle = PropertyHandle->GetChildHandle(GET_MEMBER_NAME_CHECKED(FAssemblyMetadataDesc, Type));
 
 		MetadataTypeHandle->SetOnPropertyValueChanged(FSimpleDelegate::CreateLambda([this, &MetadataDesc]()
 			{
-				ClearIncompatibleMetadataLink();
+				ValidateMetadataLink();
 
 				if ((MetadataDesc.Type == ECineAssemblyMetadataType::String) ||
 					(MetadataDesc.Type == ECineAssemblyMetadataType::AssetPath) ||
@@ -283,37 +288,40 @@ void FCineAssemblyMetadataCustomization::CustomizeChildren(TSharedRef<IPropertyH
 
 		DefaultValueWidget->AddSlot((uint8)ECineAssemblyMetadataType::CineAssembly)
 			[
-				SNew(SObjectPropertyEntryBox)
-					.AllowedClass(UCineAssembly::StaticClass())
-					.ThumbnailPool(CustomizationUtils.GetThumbnailPool())
-					.AllowCreate(true)
-					.OnShouldFilterAsset_Lambda([&MetadataDesc](const FAssetData& InAssetData) -> bool
-						{
-							if (!MetadataDesc.SchemaType.IsValid())
-							{
-								return false;
-							}
+				SNew(SAssemblyMetadataLink, Schema->TemplateSequence, MetadataDesc.Key)
+					[
+						SNew(SObjectPropertyEntryBox)
+							.AllowedClass(UCineAssembly::StaticClass())
+							.ThumbnailPool(CustomizationUtils.GetThumbnailPool())
+							.AllowCreate(true)
+							.OnShouldFilterAsset_Lambda([&MetadataDesc](const FAssetData& InAssetData) -> bool
+								{
+									if (!MetadataDesc.SchemaType.IsValid())
+									{
+										return false;
+									}
 
-							// Filter out all Cine Assembly assets that do not match the selected Schema Type for this metadata struct
-							const FAssetDataTagMapSharedView::FFindTagResult AssemblyType = InAssetData.TagsAndValues.FindTag(UCineAssembly::AssetRegistryTag_AssemblyType);
-							if (AssemblyType.IsSet())
-							{
-								return !AssemblyType.GetValue().Equals(MetadataDesc.SchemaType.GetAssetName());
-							}
-							return true;
-						})
-					.ObjectPath_Lambda([&MetadataDesc]() -> FString
-						{
-							if (const FString* Value = MetadataDesc.DefaultValue.TryGet<FString>())
-							{
-								return *Value;
-							}
-							return FString();
-						})
-					.OnObjectChanged_Lambda([&MetadataDesc](const FAssetData& InAssetData)
-						{
-							MetadataDesc.DefaultValue.Set<FString>(InAssetData.GetObjectPathString());
-						})
+									// Filter out all Cine Assembly assets that do not match the selected Schema Type for this metadata struct
+									const FAssetDataTagMapSharedView::FFindTagResult AssemblyType = InAssetData.TagsAndValues.FindTag(UCineAssembly::AssetRegistryTag_AssemblyType);
+									if (AssemblyType.IsSet())
+									{
+										return !AssemblyType.GetValue().Equals(MetadataDesc.SchemaType.GetAssetName());
+									}
+									return true;
+								})
+							.ObjectPath_Lambda([&MetadataDesc]() -> FString
+								{
+									if (const FString* Value = MetadataDesc.DefaultValue.TryGet<FString>())
+									{
+										return *Value;
+									}
+									return FString();
+								})
+							.OnObjectChanged_Lambda([&MetadataDesc](const FAssetData& InAssetData)
+								{
+									MetadataDesc.DefaultValue.Set<FString>(InAssetData.GetObjectPathString());
+								})
+					]
 			];
 
 		ChildBuilder.AddCustomRow(FText::GetEmpty())
@@ -330,7 +338,7 @@ void FCineAssemblyMetadataCustomization::CustomizeChildren(TSharedRef<IPropertyH
 	}
 }
 
-void FCineAssemblyMetadataCustomization::ClearIncompatibleMetadataLink()
+void FCineAssemblyMetadataCustomization::ValidateMetadataLink()
 {
 	if (!Schema || !Schema->TemplateSequence || !Schema->AssemblyMetadata.IsValidIndex(ArrayIndex))
 	{
@@ -338,35 +346,105 @@ void FCineAssemblyMetadataCustomization::ClearIncompatibleMetadataLink()
 	}
 
 	const FAssemblyMetadataDesc& MetadataDesc = Schema->AssemblyMetadata[ArrayIndex];
-	const FGuid* LinkedGuid = Schema->TemplateSequence->MetadataLinks.Find(MetadataDesc.Key);
-	if (!LinkedGuid)
+
+	// Get the ID of the currently linked Asset. Early-out if no asset is currently linked.
+	const FGuid* LinkedAssetID = Schema->TemplateSequence->MetadataLinks.Find(MetadataDesc.Key);
+	if (!LinkedAssetID)
 	{
 		return;
 	}
 
-	if (MetadataDesc.Type != ECineAssemblyMetadataType::AssetPath && MetadataDesc.Type != ECineAssemblyMetadataType::String)
+	// If the currently linked asset is no longer compatible with this metadata field, unlink it.
+	if (!IsLinkedAssetCompatible(*LinkedAssetID, MetadataDesc))
 	{
 		Schema->TemplateSequence->Modify();
 		Schema->TemplateSequence->MetadataLinks.Remove(MetadataDesc.Key);
-		return;
 	}
+}
 
-	const UClass* NewAllowedClass = MetadataDesc.AssetClass.ResolveClass();
-	if (!NewAllowedClass)
+bool FCineAssemblyMetadataCustomization::IsLinkedAssetCompatible(const FGuid& InLinkedAssetID, const FAssemblyMetadataDesc& InMetadataDesc) const
+{
+	if (!Schema || !Schema->TemplateSequence)
 	{
-		// No class filter means any asset is compatible — keep the link
-		return;
+		return false;
 	}
 
-	if (const FAssemblyAssociatedAssetDesc* AssetDesc = Algo::FindBy(Schema->TemplateSequence->AssociatedAssets, *LinkedGuid, &FAssemblyAssociatedAssetDesc::AssetID))
+	switch (InMetadataDesc.Type)
 	{
-		UClass* AssetClass = AssetDesc->AssetClass.LoadSynchronous();
-		if (!AssetClass || !AssetClass->IsChildOf(NewAllowedClass))
+	case ECineAssemblyMetadataType::String:
+		return true;
+
+	case ECineAssemblyMetadataType::AssetPath:
 		{
-			Schema->TemplateSequence->Modify();
-			Schema->TemplateSequence->MetadataLinks.Remove(MetadataDesc.Key);
+			// If there is no class filter, then the link is valid
+			const UClass* AllowedClass = InMetadataDesc.AssetClass.ResolveClass();
+			if (!AllowedClass)
+			{
+				return true;
+			}
+
+			// If the linked asset is one of the TemplateSequence's associated assets, check that its class passes the current asset class filter
+			if (const FAssemblyAssociatedAssetDesc* AssetDesc = Algo::FindBy(Schema->TemplateSequence->AssociatedAssets, InLinkedAssetID, &FAssemblyAssociatedAssetDesc::AssetID))
+			{
+				UClass* AssetClass = AssetDesc->AssetClass.Get();
+				return AssetClass && AssetClass->IsChildOf(AllowedClass);
+			}
+
+			// If the linked asset is not an associated asset, we can assume it must be a SubAssembly, so check that it passes the class filter
+			return UCineAssembly::StaticClass()->IsChildOf(AllowedClass);
+		}
+
+	case ECineAssemblyMetadataType::CineAssembly:
+		{
+			// CineAssembly metadata can only link to SubAssemblies, not associated assets
+			if (Algo::FindBy(Schema->TemplateSequence->AssociatedAssets, InLinkedAssetID, &FAssemblyAssociatedAssetDesc::AssetID))
+			{
+				return false;
+			}
+
+			// If there is no Schema filter, then any SubAssembly is valid
+			if (InMetadataDesc.SchemaType.IsNull())
+			{
+				return true;
+			}
+
+			// Find a SubAssemblySection that matches the linked asset ID
+			UMovieSceneSubAssemblySection* LinkedSubAssemblySection = nullptr;
+			if (UMovieScene* MovieScene = Schema->TemplateSequence->GetMovieScene())
+			{
+				for (UMovieSceneSection* Section : MovieScene->GetAllSections())
+				{
+					if (UMovieSceneSubAssemblySection* SubAssemblySection = Cast<UMovieSceneSubAssemblySection>(Section))
+					{
+						if (SubAssemblySection->GetSectionID() == InLinkedAssetID)
+						{
+							LinkedSubAssemblySection = SubAssemblySection;
+							break;
+						}
+					}
+				}
+			}
+			
+			// Validate that the Schema of this SubAssembly passes the Schema filter
+			if (LinkedSubAssemblySection)
+			{
+				const UCineAssemblySchema* TemplateSchema = LinkedSubAssemblySection->GetTemplateSchema();
+				return TemplateSchema && InMetadataDesc.SchemaType == FSoftObjectPath(TemplateSchema);
+			}
+
+			return false;
 		}
+
+	case ECineAssemblyMetadataType::Bool:
+	case ECineAssemblyMetadataType::Integer:
+	case ECineAssemblyMetadataType::Float:
+		return false;
+
+	default:
+		checkNoEntry();
 	}
+
+	return false;
 }
 
 void FCineAssemblyMetadataCustomization::CustomizeKeyProperty(TSharedRef<IPropertyHandle> PropertyHandle, IDetailChildrenBuilder& ChildBuilder, FAssemblyMetadataDesc& MetadataDesc)
diff --git a/Engine/Plugins/VirtualProduction/CinematicAssemblyTools/Source/CineAssemblyToolsEditor/Private/CineAssemblyMetadataCustomization.h b/Engine/Plugins/VirtualProduction/CinematicAssemblyTools/Source/CineAssemblyToolsEditor/Private/CineAssemblyMetadataCustomization.h
--- a/Engine/Plugins/VirtualProduction/CinematicAssemblyTools/Source/CineAssemblyToolsEditor/Private/CineAssemblyMetadataCustomization.h
+++ b/Engine/Plugins/VirtualProduction/CinematicAssemblyTools/Source/CineAssemblyToolsEditor/Private/CineAssemblyMetadataCustomization.h
@@ -26,9 +26,6 @@ private:
 	*/
 	TSharedRef<SWidget> MakeStringDefaultValueWidget(FAssemblyMetadataDesc& MetadataDesc, TSharedRef<IPropertyHandle> PropertyHandle);
 
-	/** Clears the metadata link if the linked asset's class is no longer compatible with the metadata field's AssetClass filter */
-	void ClearIncompatibleMetadataLink();
-
 	/** Customizes the metadata desc Key property */
 	void CustomizeKeyProperty(TSharedRef<IPropertyHandle> PropertyHandle, IDetailChildrenBuilder& ChildBuilder, FAssemblyMetadataDesc& MetadataDesc);
 
@@ -38,6 +35,12 @@ private:
 	/** Returns a unique key name within the context of the schema */
 	FString MakeUniqueKeyName();
 
+	/** Validates the active metadata link against the current metadata field properties and removes the link if no longer valid */
+	void ValidateMetadataLink();
+
+	/** Whether the currently linked asset is compatible with the input metadata field properties */
+	bool IsLinkedAssetCompatible(const FGuid& InLinkedAssetID, const FAssemblyMetadataDesc& InMetadataDesc) const;
+
 private:
 	/** The schema that owns the metadata desc struct being customized */
 	UCineAssemblySchema* Schema;
diff --git a/Engine/Plugins/VirtualProduction/CinematicAssemblyTools/Source/CineAssemblyToolsEditor/Private/UI/CineAssembly/SCineAssemblyAssetTreeView.cpp b/Engine/Plugins/VirtualProduction/CinematicAssemblyTools/Source/CineAssemblyToolsEditor/Private/UI/CineAssembly/SCineAssemblyAssetTreeView.cpp
--- a/Engine/Plugins/VirtualProduction/CinematicAssemblyTools/Source/CineAssemblyToolsEditor/Private/UI/CineAssembly/SCineAssemblyAssetTreeView.cpp
+++ b/Engine/Plugins/VirtualProduction/CinematicAssemblyTools/Source/CineAssemblyToolsEditor/Private/UI/CineAssembly/SCineAssemblyAssetTreeView.cpp
@@ -813,19 +813,30 @@ TSharedRef<ITableRow> SCineAssemblyAssetTreeView::OnGenerateTreeRow(TSharedPtr<I
 
 TSharedRef<SWidget> SCineAssemblyAssetTreeView::MakeMetadataLinkButton(TSharedPtr<ICineAssemblyTreeItem> TreeItem)
 {
-	TSharedPtr<FAssociatedAssetItem> AssetItem = TreeItem->AsAssociatedAsset();
-	if (!AssetItem || bIsReadOnly)
+	if (bIsReadOnly)
 	{
 		return SNullWidget::NullWidget;
 	}
 
-	const FAssemblyAssociatedAssetDesc* AssetDesc = AssetItem->FindAssetDesc();
-	if (!AssetDesc)
+	if (TSharedPtr<FAssociatedAssetItem> AssetItem = TreeItem->AsAssociatedAsset())
 	{
-		return SNullWidget::NullWidget;
+		const FAssemblyAssociatedAssetDesc* AssetDesc = AssetItem->FindAssetDesc();
+		if (!AssetDesc)
+		{
+			return SNullWidget::NullWidget;
+		}
+		return SNew(SAssemblyMetadataLink, AssetItem->GetAssembly(), AssetDesc->AssetID);
 	}
 
-	return SNew(SAssemblyMetadataLink, AssetItem->GetAssembly(), AssetDesc->AssetID);
+	if (TSharedPtr<FSubAssemblySectionItem> SectionItem = TreeItem->AsSubAssemblySection())
+	{
+		if (SectionItem->SubAssemblySection.IsValid() && SectionItem->SubAssemblySection->IsTemplateSection())
+		{
+			return SNew(SAssemblyMetadataLink, WeakAssembly.Get(), SectionItem->SubAssemblySection->GetSectionID());
+		}
+	}
+
+	return SNullWidget::NullWidget;
 }
 
 void SCineAssemblyAssetTreeView::OnGetChildren(TSharedPtr<ICineAssemblyTreeItem> TreeItem, TArray<TSharedPtr<ICineAssemblyTreeItem>>& OutNodes)
diff --git a/Engine/Plugins/VirtualProduction/CinematicAssemblyTools/Source/CineAssemblyToolsEditor/Private/UI/CineAssembly/SCineAssemblyAssetTreeView.h b/Engine/Plugins/VirtualProduction/CinematicAssemblyTools/Source/CineAssemblyToolsEditor/Private/UI/CineAssembly/SCineAssemblyAssetTreeView.h
--- a/Engine/Plugins/VirtualProduction/CinematicAssemblyTools/Source/CineAssemblyToolsEditor/Private/UI/CineAssembly/SCineAssemblyAssetTreeView.h
+++ b/Engine/Plugins/VirtualProduction/CinematicAssemblyTools/Source/CineAssemblyToolsEditor/Private/UI/CineAssembly/SCineAssemblyAssetTreeView.h
@@ -139,7 +139,7 @@ private:
 	/** Recursively expands every folder in the tree view */
 	void ExpandTreeRecursive(const TSharedPtr<FCineAssemblyFolderItem>& FolderItem) const;
 
-	/** Creates the metadata link widget for associated asset items (returns SNullWidget for non-asset items or read-only trees) */
+	/** Creates the metadata link widget for associated asset and SubAssembly section items (returns SNullWidget read-only trees) */
 	TSharedRef<SWidget> MakeMetadataLinkButton(TSharedPtr<ICineAssemblyTreeItem> TreeItem);
 
 	/** Begins a drag and drop event to drag an item out of the content tree view */
diff --git a/Engine/Plugins/VirtualProduction/CinematicAssemblyTools/Source/CineAssemblyToolsEditor/Private/UI/CineAssembly/SCineAssemblySchemaWindow.cpp b/Engine/Plugins/VirtualProduction/CinematicAssemblyTools/Source/CineAssemblyToolsEditor/Private/UI/CineAssembly/SCineAssemblySchemaWindow.cpp
--- a/Engine/Plugins/VirtualProduction/CinematicAssemblyTools/Source/CineAssemblyToolsEditor/Private/UI/CineAssembly/SCineAssemblySchemaWindow.cpp
+++ b/Engine/Plugins/VirtualProduction/CinematicAssemblyTools/Source/CineAssemblyToolsEditor/Private/UI/CineAssembly/SCineAssemblySchemaWindow.cpp
@@ -3,6 +3,7 @@
 #include "SCineAssemblySchemaWindow.h"
 
 #include "Algo/Contains.h"
+#include "Algo/Find.h"
 #include "AssetDefinitionRegistry.h"
 #include "AssetRegistry/AssetRegistryModule.h"
 #include "AssetToolsModule.h"
@@ -764,19 +765,12 @@ void SCineAssemblySchemaWindow::OnSchemaPropertiesChanged(const FPropertyChanged
 			CineAssemblyNamingTokens->AddMetadataToken(MetadataDesc.Key);
 		}
 	}
-	else if (PropertyChangedEvent.GetMemberPropertyName() == GET_MEMBER_NAME_CHECKED(UCineAssemblySchema, AssemblyMetadata)	&& PropertyChangedEvent.ChangeType == EPropertyChangeType::ArrayRemove)
+	else if (PropertyChangedEvent.GetMemberPropertyName() == GET_MEMBER_NAME_CHECKED(UCineAssemblySchema, AssemblyMetadata) && PropertyChangedEvent.ChangeType == EPropertyChangeType::ArrayRemove)
 	{
-		// Clean up MetadataLinks entries that reference metadata keys no longer in the schema
+		// Validate the current metadata links in case the metadata field that was removed had an active link
 		if (SchemaToConfigure->TemplateSequence)
 		{
-			SchemaToConfigure->TemplateSequence->Modify();
-			for (auto It = SchemaToConfigure->TemplateSequence->MetadataLinks.CreateIterator(); It; ++It)
-			{
-				if (!Algo::ContainsBy(SchemaToConfigure->AssemblyMetadata, It->Key, &FAssemblyMetadataDesc::Key))
-				{
-					It.RemoveCurrent();
-				}
-			}
+			SchemaToConfigure->TemplateSequence->ValidateMetadataLinks();
 		}
 	}
 	else if (PropertyName == GET_MEMBER_NAME_CHECKED(UCineAssemblySchema, ThumbnailImage))
@@ -814,12 +808,62 @@ void SCineAssemblySchemaWindow::OnMovieSceneDataChanged(EMovieSceneDataChangeTyp
 {
 	MuteTracks();
 
+	// Validate the current metadata links in case a SubAssembly track/section was recently modified/removed
+	if (SchemaToConfigure->TemplateSequence)
+	{
+		SchemaToConfigure->TemplateSequence->ValidateMetadataLinks();
+		ValidateSubAssemblyLinks();
+	}
+
 	if (TreeView)
 	{
 		TreeView->Reinitialize();
 	}
 }
 
+void SCineAssemblySchemaWindow::ValidateSubAssemblyLinks()
+{
+	UCineAssembly* TemplateSequence = SchemaToConfigure->TemplateSequence;
+	if (!TemplateSequence || TemplateSequence->MetadataLinks.IsEmpty())
+	{
+		return;
+	}
+
+	UMovieScene* MovieScene = TemplateSequence->GetMovieScene();
+	if (!MovieScene)
+	{
+		return;
+	}
+
+	for (auto It = TemplateSequence->MetadataLinks.CreateIterator(); It; ++It)
+	{
+		// Only check CineAssembly-type metadata fields with a schema filter
+		const FAssemblyMetadataDesc* MetadataDesc = Algo::FindBy(SchemaToConfigure->AssemblyMetadata, It->Key, &FAssemblyMetadataDesc::Key);
+		if (!MetadataDesc || MetadataDesc->Type != ECineAssemblyMetadataType::CineAssembly || MetadataDesc->SchemaType.IsNull())
+		{
+			continue;
+		}
+
+		// Find the linked SubAssembly section and check if its schema still matches
+		for (UMovieSceneSection* Section : MovieScene->GetAllSections())
+		{
+			if (UMovieSceneSubAssemblySection* SubAssemblySection = Cast<UMovieSceneSubAssemblySection>(Section))
+			{
+				if (SubAssemblySection->GetSectionID() == It->Value)
+				{
+					const UCineAssemblySchema* TemplateSchema = SubAssemblySection->GetTemplateSchema();
+					if (!TemplateSchema || MetadataDesc->SchemaType != FSoftObjectPath(TemplateSchema))
+					{
+						TemplateSequence->Modify();
+						It.RemoveCurrent();
+					}
+					break;
+				}
+			}
+		}
+	}
+}
+
 void SCineAssemblySchemaWindow::OnActorAddedToSequencer(AActor* InActor, const FGuid InBinding)
 {
 	// Possessables in our template sequence may not actually bind to any actors in the world
diff --git a/Engine/Plugins/VirtualProduction/CinematicAssemblyTools/Source/CineAssemblyToolsEditor/Private/UI/CineAssembly/SCineAssemblySchemaWindow.h b/Engine/Plugins/VirtualProduction/CinematicAssemblyTools/Source/CineAssemblyToolsEditor/Private/UI/CineAssembly/SCineAssemblySchemaWindow.h
--- a/Engine/Plugins/VirtualProduction/CinematicAssemblyTools/Source/CineAssemblyToolsEditor/Private/UI/CineAssembly/SCineAssemblySchemaWindow.h
+++ b/Engine/Plugins/VirtualProduction/CinematicAssemblyTools/Source/CineAssemblyToolsEditor/Private/UI/CineAssembly/SCineAssemblySchemaWindow.h
@@ -113,6 +113,9 @@ private:
 	/** Callback whenever MovieScene data in the template sequence is changed */
 	void OnMovieSceneDataChanged(EMovieSceneDataChangeType InChangeType);
 
+	/** Validates all CineAssembly metadata links */
+	void ValidateSubAssemblyLinks();
+
 	/** Extends the sequencer widget's toolbar */
 	void AddSequencerToolbarExtension(FToolBarBuilder& Builder);
 
diff --git a/Engine/Plugins/VirtualProduction/CinematicAssemblyTools/Source/CineAssemblyToolsEditor/Private/UI/SAssemblyMetadataLink.cpp b/Engine/Plugins/VirtualProduction/CinematicAssemblyTools/Source/CineAssemblyToolsEditor/Private/UI/SAssemblyMetadataLink.cpp
--- a/Engine/Plugins/VirtualProduction/CinematicAssemblyTools/Source/CineAssemblyToolsEditor/Private/UI/SAssemblyMetadataLink.cpp
+++ b/Engine/Plugins/VirtualProduction/CinematicAssemblyTools/Source/CineAssemblyToolsEditor/Private/UI/SAssemblyMetadataLink.cpp
@@ -7,6 +7,10 @@
 #include "CineAssembly.h"
 #include "CineAssemblySchema.h"
 #include "Framework/MultiBox/MultiBoxBuilder.h"
+#include "MovieScene.h"
+#include "MovieSceneSubAssemblySection.h"
+#include "MovieSceneSubAssemblyTrack.h"
+#include "Sections/MovieSceneSubSection.h"
 #include "Styling/StyleColors.h"
 #include "Widgets/Images/SImage.h"
 #include "Widgets/Input/SComboButton.h"
@@ -30,11 +34,6 @@ void SAssemblyMetadataLink::Construct(const FArguments& InArgs, UCineAssembly* I
 	AssociatedAssetID = InAssetID;
 	bIsAssetMode = true;
 
-	if (const FAssemblyAssociatedAssetDesc* AssetDesc = Algo::FindBy(InAssembly->AssociatedAssets, AssociatedAssetID, &FAssemblyAssociatedAssetDesc::AssetID))
-	{
-		CachedAssetClass = AssetDesc->AssetClass.Get();
-	}
-
 	ConstructInternal(InArgs);
 }
 
@@ -50,23 +49,6 @@ void SAssemblyMetadataLink::Construct(const FArguments& InArgs, UCineAssembly* I
 	MetadataKey = InMetadataKey;
 	bIsAssetMode = false;
 
-	if (const UCineAssemblySchema* Schema = GetSchema())
-	{
-		if (const FAssemblyMetadataDesc* MetadataDesc = Algo::FindBy(Schema->AssemblyMetadata, MetadataKey, &FAssemblyMetadataDesc::Key))
-		{
-			switch(MetadataDesc->Type)
-			{
-			case ECineAssemblyMetadataType::AssetPath:
-				CachedAssetClass = MetadataDesc->AssetClass.ResolveClass();
-				break;
-			default:
-				// For all other metadata types, the AssetClass is not used to filter out associated assets that are compatible for linking
-				CachedAssetClass = nullptr;
-				break;
-			}
-		}
-	}
-
 	ConstructInternal(InArgs);
 }
 
@@ -170,38 +152,56 @@ TSharedRef<SWidget> SAssemblyMetadataLink::BuildMenuContent()
 			);
 		}
 	}
-	else if (IsMetadataKeyValid())
+	else if (const FAssemblyMetadataDesc* MetadataDesc = Algo::FindBy(Schema->AssemblyMetadata, MetadataKey, &FAssemblyMetadataDesc::Key))
 	{
-		// Metadata Mode: list compatible associated assets as radio buttons (a metadata field can only be linked to one associated asset)
-		for (const FAssemblyAssociatedAssetDesc& AssetDesc : Assembly->AssociatedAssets)
+		// Metadata Mode: list compatible associated assets as radio buttons (a metadata field can only be linked to one asset)
+		// Associated assets are only compatible with String and AssetPath metadata types, not CineAssembly
+		if (MetadataDesc->Type == ECineAssemblyMetadataType::String || MetadataDesc->Type == ECineAssemblyMetadataType::AssetPath)
 		{
-			if (!IsAssetClassCompatible(AssetDesc.AssetClass.Get(), CachedAssetClass.Get()))
+			const UClass* AllowedClass = MetadataDesc->AssetClass.ResolveClass();
+			for (const FAssemblyAssociatedAssetDesc& AssetDesc : Assembly->AssociatedAssets)
 			{
-				continue;
+				if (IsAssetClassCompatible(AssetDesc.AssetClass.Get(), AllowedClass))
+				{
+					AddMenuEntryForAsset(MenuBuilder, AssetDesc.AssetID, FText::FromString(AssetDesc.AssetName.Template));
+				}
 			}
-
-			const FText Tooltip = IsAssetLinked(MetadataKey, AssetDesc.AssetID)
-				? FText::Format(LOCTEXT("UnlinkAssetTooltip", "Unlink from '{0}'"), FText::FromString(AssetDesc.AssetName.Template))
-				: FText::Format(LOCTEXT("LinkAssetTooltip", "Link to '{0}'"), FText::FromString(AssetDesc.AssetName.Template));
-
-			MenuBuilder.AddMenuEntry(
-				FText::FromString(AssetDesc.AssetName.Template),
-				Tooltip,
-				FSlateIcon(),
-				FUIAction(
-					FExecuteAction::CreateSP(this, &SAssemblyMetadataLink::ToggleMetadataLink, MetadataKey, AssetDesc.AssetID),
-					FCanExecuteAction(),
-					FIsActionChecked::CreateSP(this, &SAssemblyMetadataLink::IsAssetLinked, MetadataKey, AssetDesc.AssetID)
-				),
-				NAME_None,
-				EUserInterfaceActionType::RadioButton
-			);
 		}
+
+		// Also list compatible SubAssemblies
+		ForEachSubAssembly([this, &MenuBuilder, &MetadataDesc](FGuid SubAssemblyID, const FText& DisplayName, const UCineAssemblySchema*)
+		{
+			if (IsSubAssemblyCompatible(SubAssemblyID, *MetadataDesc))
+			{
+				AddMenuEntryForAsset(MenuBuilder, SubAssemblyID, DisplayName);
+			}
+			return true;
+		});
 	}
 
 	return MenuBuilder.MakeWidget();
 }
 
+void SAssemblyMetadataLink::AddMenuEntryForAsset(FMenuBuilder& MenuBuilder, FGuid AssetID, const FText& DisplayName)
+{
+	const FText Tooltip = IsAssetLinked(MetadataKey, AssetID)
+		? FText::Format(LOCTEXT("UnlinkTooltip", "Unlink from '{0}'"), DisplayName)
+		: FText::Format(LOCTEXT("LinkTooltip", "Link to '{0}'"), DisplayName);
+
+	MenuBuilder.AddMenuEntry(
+		DisplayName,
+		Tooltip,
+		FSlateIcon(),
+		FUIAction(
+			FExecuteAction::CreateSP(this, &SAssemblyMetadataLink::ToggleMetadataLink, MetadataKey, AssetID),
+			FCanExecuteAction(),
+			FIsActionChecked::CreateSP(this, &SAssemblyMetadataLink::IsAssetLinked, MetadataKey, AssetID)
+		),
+		NAME_None,
+		EUserInterfaceActionType::RadioButton
+	);
+}
+
 void SAssemblyMetadataLink::ToggleMetadataLink(FString InMetadataKey, FGuid InAssetID)
 {
 	if (UCineAssembly* Assembly = WeakAssembly.Get())
@@ -248,9 +248,10 @@ bool SAssemblyMetadataLink::HasAnyLink() const
 		}
 		return false;
 	}
-	else if (IsMetadataKeyValid())
+	else if (const UCineAssemblySchema* Schema = GetSchema())
 	{
-		return Assembly->MetadataLinks.Contains(MetadataKey);
+		const bool bKeyIsValid = Algo::ContainsBy(Schema->AssemblyMetadata, MetadataKey, &FAssemblyMetadataDesc::Key);
+		return bKeyIsValid && Assembly->MetadataLinks.Contains(MetadataKey);
 	}
 
 	return false;
@@ -271,16 +272,34 @@ EVisibility SAssemblyMetadataLink::GetLinkButtonVisibility() const
 		const bool bHasCompatibleField = Schema->AssemblyMetadata.ContainsByPredicate([this](const FAssemblyMetadataDesc& MetadataDesc) { return IsMetadataDescCompatible(MetadataDesc); });
 		return bHasCompatibleField ? EVisibility::Visible : EVisibility::Collapsed;
 	}
-	else
+	else if (const FAssemblyMetadataDesc* MetadataDesc = Algo::FindBy(Schema->AssemblyMetadata, MetadataKey, &FAssemblyMetadataDesc::Key))
 	{
-		if (!IsMetadataKeyValid())
+		// Check if any associated assets are compatible (only for String and AssetPath types)
+		if (MetadataDesc->Type == ECineAssemblyMetadataType::String || MetadataDesc->Type == ECineAssemblyMetadataType::AssetPath)
 		{
-			return EVisibility::Collapsed;
+			const UClass* AllowedClass = MetadataDesc->AssetClass.ResolveClass();
+			const bool bHasCompatibleAsset = Assembly->AssociatedAssets.ContainsByPredicate([this, AllowedClass](const FAssemblyAssociatedAssetDesc& AssetDesc) { return IsAssetClassCompatible(AssetDesc.AssetClass.Get(), AllowedClass); });
+			if (bHasCompatibleAsset)
+			{
+				return EVisibility::Visible;
+			}
 		}
 
-		const bool bHasCompatibleAsset = Assembly->AssociatedAssets.ContainsByPredicate([this](const FAssemblyAssociatedAssetDesc& AssetDesc) { return IsAssetClassCompatible(AssetDesc.AssetClass.Get(), CachedAssetClass.Get()); });
-		return bHasCompatibleAsset ? EVisibility::Visible : EVisibility::Collapsed;
+		// Check if any SubAssemblies are compatible
+		bool bHasCompatibleSubAssembly = false;
+		ForEachSubAssembly([this, MetadataDesc, &bHasCompatibleSubAssembly](FGuid SubAssemblyID, const FText&, const UCineAssemblySchema*)
+		{
+			bHasCompatibleSubAssembly = IsSubAssemblyCompatible(SubAssemblyID, *MetadataDesc);
+			return !bHasCompatibleSubAssembly; // Stop on first match
+		});
+
+		if (bHasCompatibleSubAssembly)
+		{
+			return EVisibility::Visible;
+		}
 	}
+
+	return EVisibility::Collapsed;
 }
 
 FSlateColor SAssemblyMetadataLink::GetLinkButtonColor() const
@@ -296,23 +315,37 @@ FText SAssemblyMetadataLink::GetLinkedAssetName() const
 		return FText::GetEmpty();
 	}
 
-	if (!IsMetadataKeyValid())
-	{
-		return FText::GetEmpty();
-	}
-
 	if (UCineAssembly* Assembly = WeakAssembly.Get())
 	{
-		// In Metadata Mode, show the linked asset's template name
+		// In Metadata Mode, show the linked target's template name
 		if (const FGuid* LinkedID = Assembly->MetadataLinks.Find(MetadataKey))
 		{
+			// Check associated assets first
 			if (const FAssemblyAssociatedAssetDesc* AssetDesc = Algo::FindBy(Assembly->AssociatedAssets, *LinkedID, &FAssemblyAssociatedAssetDesc::AssetID))
 			{
 				return FText::FromString(AssetDesc->AssetName.Template);
 			}
-		}
 
-		return LOCTEXT("LinkedAssetNotFound", "(Linked asset not found)");
+			// Check SubAssemblies
+			FText SubAssemblyName;
+			ForEachSubAssembly([&LinkedID, &SubAssemblyName](FGuid ID, const FText& DisplayName, const UCineAssemblySchema*)
+			{
+				if (ID == *LinkedID)
+				{
+					SubAssemblyName = DisplayName;
+					return false;
+				}
+				return true;
+			});
+
+			if (!SubAssemblyName.IsEmpty())
+			{
+				return SubAssemblyName;
+			}
+
+			// Link exists but target not found
+			return LOCTEXT("LinkedAssetNotFound", "(Linked target not found)");
+		}
 	}
 
 	return FText::GetEmpty();
@@ -334,28 +367,34 @@ const UCineAssemblySchema* SAssemblyMetadataLink::GetSchema() const
 	return Assembly->GetSchema();
 }
 
-bool SAssemblyMetadataLink::IsMetadataKeyValid() const
+bool SAssemblyMetadataLink::IsAssetClassCompatible(const UClass* InAssetClass, const UClass* InFilterClass) const
 {
-	if (bIsAssetMode)
-	{
-		return false;
-	}
-
-	const UCineAssemblySchema* Schema = GetSchema();
-	if (!Schema)
+	if (!InFilterClass)
 	{
-		return false;
+		return true;
 	}
-	return Algo::ContainsBy(Schema->AssemblyMetadata, MetadataKey, &FAssemblyMetadataDesc::Key);
+	return InAssetClass && InAssetClass->IsChildOf(InFilterClass);
 }
 
-bool SAssemblyMetadataLink::IsAssetClassCompatible(const UClass* InAssetClass, const UClass* InFilterClass) const
+bool SAssemblyMetadataLink::IsSchemaCompatible(FGuid InSubAssemblyID, const FSoftObjectPath& InSchemaType) const
 {
-	if (!InFilterClass)
+	if (!InSchemaType.IsValid())
 	{
 		return true;
 	}
-	return InAssetClass && InAssetClass->IsChildOf(InFilterClass);
+
+	bool bIsCompatible = false;
+	ForEachSubAssembly([&](FGuid ID, const FText&, const UCineAssemblySchema* Schema)
+	{
+		if (ID == InSubAssemblyID)
+		{
+			bIsCompatible = Schema && InSchemaType == FSoftObjectPath(Schema);
+			return false;
+		}
+		return true;
+	});
+
+	return bIsCompatible;
 }
 
 bool SAssemblyMetadataLink::IsMetadataDescCompatible(const FAssemblyMetadataDesc& InMetadataDesc) const
@@ -366,12 +405,59 @@ bool SAssemblyMetadataLink::IsMetadataDescCompatible(const FAssemblyMetadataDesc
 		return true;
 
 	case ECineAssemblyMetadataType::AssetPath:
-		return IsAssetClassCompatible(CachedAssetClass.Get(), InMetadataDesc.AssetClass.ResolveClass());
+		{
+			const UClass* AssetClass = nullptr;
+			if (UCineAssembly* Assembly = WeakAssembly.Get())
+			{
+				if (const FAssemblyAssociatedAssetDesc* AssetDesc = Algo::FindBy(Assembly->AssociatedAssets, AssociatedAssetID, &FAssemblyAssociatedAssetDesc::AssetID))
+				{
+					AssetClass = AssetDesc->AssetClass.Get();
+				}
+			}
+			return IsAssetClassCompatible(AssetClass, InMetadataDesc.AssetClass.ResolveClass());
+		}
+
+	case ECineAssemblyMetadataType::CineAssembly:
+		{
+			// CineAssembly metadata can only link to SubAssemblies, not associated assets
+			if (UCineAssembly* Assembly = WeakAssembly.Get())
+			{
+				if (Algo::FindBy(Assembly->AssociatedAssets, AssociatedAssetID, &FAssemblyAssociatedAssetDesc::AssetID))
+				{
+					return false;
+				}
+			}
+			return IsSchemaCompatible(AssociatedAssetID, InMetadataDesc.SchemaType);
+		}
 
 	case ECineAssemblyMetadataType::Bool:
 	case ECineAssemblyMetadataType::Integer:
 	case ECineAssemblyMetadataType::Float:
+		return false;
+
+	default:
+		checkNoEntry();
+	}
+
+	return false;
+}
+
+bool SAssemblyMetadataLink::IsSubAssemblyCompatible(FGuid InSubAssemblyID, const FAssemblyMetadataDesc& InMetadataDesc) const
+{
+	switch (InMetadataDesc.Type)
+	{
+	case ECineAssemblyMetadataType::String:
+		return true;
+
+	case ECineAssemblyMetadataType::AssetPath:
+		return IsAssetClassCompatible(UCineAssembly::StaticClass(), InMetadataDesc.AssetClass.ResolveClass());
+
 	case ECineAssemblyMetadataType::CineAssembly:
+		return IsSchemaCompatible(InSubAssemblyID, InMetadataDesc.SchemaType);
+
+	case ECineAssemblyMetadataType::Bool:
+	case ECineAssemblyMetadataType::Integer:
+	case ECineAssemblyMetadataType::Float:
 		return false;
 
 	default:
@@ -381,4 +467,49 @@ bool SAssemblyMetadataLink::IsMetadataDescCompatible(const FAssemblyMetadataDesc
 	return false;
 }
 
+void SAssemblyMetadataLink::ForEachSubAssembly(TFunctionRef<bool(FGuid ID, const FText& DisplayName, const UCineAssemblySchema* Schema)> Visitor) const
+{
+	UCineAssembly* Assembly = WeakAssembly.Get();
+	if (!Assembly)
+	{
+		return;
+	}
+
+	// If the Assembly is the TemplateSequence of a Schema, then we traverse the SubAssembly Sections. 
+	// Otherwise, the Assembly is a real configured object, so we traverse the managed SubAssemblies list. 
+	const bool bIsTemplateAssembly = Assembly->GetTypedOuter<UCineAssemblySchema>() != nullptr;
+	if (bIsTemplateAssembly)
+	{
+		if (UMovieScene* MovieScene = Assembly->GetMovieScene())
+		{
+			for (UMovieSceneSection* Section : MovieScene->GetAllSections())
+			{
+				if (UMovieSceneSubAssemblySection* SubAssemblySection = Cast<UMovieSceneSubAssemblySection>(Section))
+				{
+					if (SubAssemblySection->IsTemplateSection())
+					{
+						if (!Visitor(SubAssemblySection->GetSectionID(), SubAssemblySection->GetSequenceName(), SubAssemblySection->GetTemplateSchema()))
+						{
+							return;
+						}
+					}
+				}
+			}
+		}
+	}
+	else
+	{
+		for (UMovieSceneSubSection* SubSection : Assembly->SubAssemblies)
+		{
+			if (UCineAssembly* SubAssembly = Cast<UCineAssembly>(SubSection ? SubSection->GetSequence() : nullptr))
+			{
+				if (!Visitor(SubAssembly->GetAssemblyGuid(), FText::FromString(SubAssembly->AssemblyName.Template), SubAssembly->GetSchema()))
+				{
+					return;
+				}
+			}
+		}
+	}
+}
+
 #undef LOCTEXT_NAMESPACE
diff --git a/Engine/Plugins/VirtualProduction/CinematicAssemblyTools/Source/CineAssemblyToolsEditor/Private/UI/SAssemblyMetadataLink.h b/Engine/Plugins/VirtualProduction/CinematicAssemblyTools/Source/CineAssemblyToolsEditor/Private/UI/SAssemblyMetadataLink.h
--- a/Engine/Plugins/VirtualProduction/CinematicAssemblyTools/Source/CineAssemblyToolsEditor/Private/UI/SAssemblyMetadataLink.h
+++ b/Engine/Plugins/VirtualProduction/CinematicAssemblyTools/Source/CineAssemblyToolsEditor/Private/UI/SAssemblyMetadataLink.h
@@ -4,6 +4,7 @@
 
 #include "Widgets/SCompoundWidget.h"
 
+class FMenuBuilder;
 class UCineAssembly;
 class UCineAssemblySchema;
 struct FAssemblyAssociatedAssetDesc;
@@ -43,6 +44,9 @@ private:
 	/** Builds the dropdown menu content showing the link options */
 	TSharedRef<SWidget> BuildMenuContent();
 
+	/** Adds a menu entry to the dropdown menu content for the input Asset */
+	void AddMenuEntryForAsset(FMenuBuilder& MenuBuilder, FGuid AssetID, const FText& DisplayName);
+
 	/** Toggles a link between a metadata key and an asset GUID */
 	void ToggleMetadataLink(FString InMetadataKey, FGuid InAssetID);
 
@@ -67,15 +71,21 @@ private:
 	/** Returns the schema that defines the assembly's metadata fields */
 	const UCineAssemblySchema* GetSchema() const;
 
-	/** Whether the cached metadata key is still in the schema's list of metadata fields */
-	bool IsMetadataKeyValid() const;
-
 	/** Whether the AssetClass is compatible with the FilterClass. Always returns true if FilterClass is nullptr. */
 	bool IsAssetClassCompatible(const UClass* InAssetClass, const UClass* InFilterClass) const;
 
+	/** Whether the input SubAssembly has a schema that is compatible with the input SchemaType filter. Always return true if SchemaType is not specified. */
+	bool IsSchemaCompatible(FGuid InSubAssemblyID, const FSoftObjectPath& InSchemaType) const;
+
 	/** Whether the input metadata desc is compatible for linking with this asset */
 	bool IsMetadataDescCompatible(const FAssemblyMetadataDesc& InMetadataDesc) const;
 
+	/** Whether the input SubAssembly is compatible for linking with the input metadata desc */
+	bool IsSubAssemblyCompatible(FGuid InSubAssemblyID, const FAssemblyMetadataDesc& InMetadataDesc) const;
+
+	/** Iterates all SubAssemblies of the tracked Assembly and calls the visitor for each */
+	void ForEachSubAssembly(TFunctionRef<bool(FGuid ID, const FText& DisplayName, const UCineAssemblySchema* Schema)> Visitor) const;
+
 private:
 	/** The assembly whose associated assets and schema metadata are being linked */
 	TWeakObjectPtr<UCineAssembly> WeakAssembly;
@@ -88,7 +98,4 @@ private:
 
 	/** Whether this widget is in Asset Mode (true) or Metadata Mode (false) */
 	bool bIsAssetMode = false;
-
-	/** Cached asset class — in Asset Mode: the associated asset's class; in Metadata Mode: the metadata field's allowed class */
-	TWeakObjectPtr<const UClass> CachedAssetClass;
 };
